[UMG] UMG Preview Window

diff --git a/Engine/Plugins/Editor/UMGWidgetPreview/Source/UMGWidgetPreview/Private/AssetDefinition_WidgetPreview.cpp b/Engine/Plugins/Editor/UMGWidgetPreview/Source/UMGWidgetPreview/Private/AssetDefinition_WidgetPreview.cpp
new file mode 100644
--- /dev/null
+++ b/Engine/Plugins/Editor/UMGWidgetPreview/Source/UMGWidgetPreview/Private/AssetDefinition_WidgetPreview.cpp
@@ -0,0 +1,64 @@
+// Copyright Epic Games, Inc. All Rights Reserved.
+
+#include "AssetDefinition_WidgetPreview.h"
+
+#include "Editor.h"
+#include "Engine/Blueprint.h"
+#include "WidgetPreview.h"
+#include "WidgetPreviewEditor.h"
+
+#define LOCTEXT_NAMESPACE "AssetTypeActions"
+
+UAssetDefinition_WidgetPreview::UAssetDefinition_WidgetPreview() = default;
+
+UAssetDefinition_WidgetPreview::~UAssetDefinition_WidgetPreview() = default;
+
+FText UAssetDefinition_WidgetPreview::GetAssetDisplayName() const
+{
+	return LOCTEXT("AssetTypeActions_WidgetPreview", "Widget Preview");
+}
+
+FLinearColor UAssetDefinition_WidgetPreview::GetAssetColor() const
+{
+	return FLinearColor(FColor(44, 89, 180));
+}
+
+TSoftClassPtr<> UAssetDefinition_WidgetPreview::GetAssetClass() const
+{
+	return UWidgetPreview::StaticClass();
+}
+
+TConstArrayView<FAssetCategoryPath> UAssetDefinition_WidgetPreview::GetAssetCategories() const
+{
+	static const TArray<FAssetCategoryPath, TFixedAllocator<1>> Categories = { EAssetCategoryPaths::UI };
+	return Categories;
+}
+
+EAssetCommandResult UAssetDefinition_WidgetPreview::OpenAssets(const FAssetOpenArgs& OpenArgs) const
+{
+	EAssetCommandResult Result = EAssetCommandResult::Unhandled;
+
+	for (UWidgetPreview* WidgetPreview : OpenArgs.LoadObjects<UWidgetPreview>())
+	{
+		UAssetEditorSubsystem* AssetEditorSubsystem = GEditor->GetEditorSubsystem<UAssetEditorSubsystem>();
+		UWidgetPreviewEditor* AssetEditor = NewObject<UWidgetPreviewEditor>(AssetEditorSubsystem, NAME_None, RF_Transient);
+		AssetEditor->Initialize(WidgetPreview);
+		Result = EAssetCommandResult::Handled;
+	}
+
+	return Result;
+}
+
+FText UAssetDefinition_WidgetPreview::GetAssetDescription(const FAssetData& AssetData) const
+{
+	FString Description = AssetData.GetTagValueRef<FString>(GET_MEMBER_NAME_CHECKED(UBlueprint, BlueprintDescription));
+	if (!Description.IsEmpty())
+	{
+		Description.ReplaceInline( TEXT( "\\n" ), TEXT( "\n" ) );
+		return FText::FromString( MoveTemp(Description) );
+	}
+
+	return FText::GetEmpty();
+}
+
+#undef LOCTEXT_NAMESPACE
diff --git a/Engine/Plugins/Editor/UMGWidgetPreview/Source/UMGWidgetPreview/Private/AssetDefinition_WidgetPreview.h b/Engine/Plugins/Editor/UMGWidgetPreview/Source/UMGWidgetPreview/Private/AssetDefinition_WidgetPreview.h
new file mode 100644
--- /dev/null
+++ b/Engine/Plugins/Editor/UMGWidgetPreview/Source/UMGWidgetPreview/Private/AssetDefinition_WidgetPreview.h
@@ -0,0 +1,27 @@
+// Copyright Epic Games, Inc. All Rights Reserved.
+
+#pragma once
+
+#include "AssetDefinitionDefault.h"
+
+#include "AssetDefinition_WidgetPreview.generated.h"
+
+UCLASS()
+class UAssetDefinition_WidgetPreview
+	: public UAssetDefinitionDefault
+{
+	GENERATED_BODY()
+
+public:
+	UAssetDefinition_WidgetPreview();
+	virtual ~UAssetDefinition_WidgetPreview() override;
+
+	//~ Begin UAssetDefinition
+	virtual FText GetAssetDisplayName() const override;
+	virtual FLinearColor GetAssetColor() const override;
+	virtual TSoftClassPtr<UObject> GetAssetClass() const override;
+	virtual TConstArrayView<FAssetCategoryPath> GetAssetCategories() const override;
+	virtual EAssetCommandResult OpenAssets(const FAssetOpenArgs& OpenArgs) const override;
+	virtual FText GetAssetDescription(const FAssetData& AssetData) const override;
+	//~ End UAssetDefinition
+};
diff --git a/Engine/Plugins/Editor/UMGWidgetPreview/Source/UMGWidgetPreview/Private/UMGWidgetPreviewModule.cpp b/Engine/Plugins/Editor/UMGWidgetPreview/Source/UMGWidgetPreview/Private/UMGWidgetPreviewModule.cpp
new file mode 100644
--- /dev/null
+++ b/Engine/Plugins/Editor/UMGWidgetPreview/Source/UMGWidgetPreview/Private/UMGWidgetPreviewModule.cpp
@@ -0,0 +1,118 @@
+﻿// Copyright Epic Games, Inc. All Rights Reserved.
+
+#include "UMGWidgetPreviewModule.h"
+
+#include "Blueprint/UserWidget.h"
+#include "ContentBrowserMenuContexts.h"
+#include "Editor.h"
+#include "MessageLogModule.h"
+#include "Modules/ModuleManager.h"
+#include "PropertyEditorModule.h"
+#include "ToolMenus.h"
+#include "UObject/UObjectHash.h"
+#include "WidgetBlueprint.h"
+#include "WidgetPreviewCommands.h"
+#include "WidgetPreviewEditor.h"
+#include "WidgetPreviewLog.h"
+#include "WidgetPreviewStyle.h"
+#include "WidgetPreview.h"
+
+#define LOCTEXT_NAMESPACE "UMGWidgetPreviewModule"
+
+namespace UE::UMGWidgetPreview::Private
+{
+	void FUMGWidgetPreviewModule::StartupModule()
+	{
+		FWidgetPreviewCommands::Register();
+
+		FMessageLogModule& MessageLogModule = FModuleManager::LoadModuleChecked<FMessageLogModule>("MessageLog");
+		FMessageLogInitializationOptions InitOptions;
+		InitOptions.bShowFilters = false;
+		InitOptions.bShowPages = false;
+		InitOptions.bShowInLogWindow = false;
+		InitOptions.bAllowClear = true;
+		MessageLogModule.RegisterLogListing(MessageLogName, LOCTEXT("WidgetPreviewLog", "Widget Preview Log"), InitOptions);
+
+		// Menus need to be registered in a callback to make sure the system is ready for them.
+		UToolMenus::RegisterStartupCallback(
+			FSimpleMulticastDelegate::FDelegate::CreateRaw(this, &FUMGWidgetPreviewModule::RegisterMenus));
+	}
+
+	void FUMGWidgetPreviewModule::ShutdownModule()
+	{
+		UToolMenus::UnRegisterStartupCallback(this);
+
+		FMessageLogModule& MessageLogModule = FModuleManager::LoadModuleChecked<FMessageLogModule>("MessageLog");
+		MessageLogModule.UnregisterLogListing(MessageLogName);
+
+		FWidgetPreviewCommands::Unregister();
+	}
+
+	void FUMGWidgetPreviewModule::RegisterMenus()
+	{
+		// Allows cleanup when module unloads.
+		FToolMenuOwnerScoped OwnerScoped(this);
+
+		// Extend the content browser context menu for widgets
+		{
+			UToolMenu* Menu = UToolMenus::Get()->ExtendMenu("ContentBrowser.AssetContextMenu.WidgetBlueprint");
+			FToolMenuSection& Section = Menu->FindOrAddSection("GetAssetActions");
+			Section.AddDynamicEntry("OpenWidgetPreviewEditor", FNewToolMenuSectionDelegate::CreateLambda(
+				[](FToolMenuSection& Section)
+				{
+					// We'll need to get the target assets out of the context
+					if (UContentBrowserAssetContextMenuContext* Context = Section.FindContext<UContentBrowserAssetContextMenuContext>())
+					{
+						// We are deliberately not using Context->GetSelectedObjects() here to avoid triggering a load from right-clicking
+						// an asset in the content browser.
+						if (UWidgetPreviewEditor::AreAssetsValidTargets(Context->SelectedAssets))
+						{
+							const TSharedPtr<class FUICommandList> CommandListToBind = MakeShared<FUICommandList>();
+							CommandListToBind->MapAction(
+								FWidgetPreviewCommands::Get().OpenEditor,
+								FExecuteAction::CreateLambda(
+									[WeakContext = TWeakObjectPtr<UContentBrowserAssetContextMenuContext>(Context)]()
+									{
+										check(GEditor);
+
+										if (UContentBrowserAssetContextMenuContext* Context = WeakContext.Get())
+										{
+											// When we actually do want to open the editor, trigger the load to get the objects
+											TArray<TObjectPtr<UObject>> ObjectsToEdit;
+											ObjectsToEdit.Append(Context->LoadSelectedObjects<UObject>());
+
+											// If we fail the ensure here, then there must be something that we're failing to check properly
+											// in AreAssetsValidTargets that we would need to track down and check in the asset data.
+											if (ensure(UWidgetPreviewEditor::AreObjectsValidTargets(ObjectsToEdit)))
+											{
+												for (const TObjectPtr<UObject>& ObjectToEdit : ObjectsToEdit)
+												{
+													if (const UWidgetBlueprint* AsWidgetBlueprint = Cast<UWidgetBlueprint>(ObjectToEdit))
+													{
+														UWidgetPreview* PreviewForWidget = UWidgetPreviewEditor::CreatePreviewForWidget(
+															Cast<UUserWidget>(AsWidgetBlueprint->GeneratedClass->GetDefaultObject<UUserWidget>()));
+														GEditor->GetEditorSubsystem<UAssetEditorSubsystem>()->OpenEditorForAsset(PreviewForWidget);
+													}
+												}
+											}
+										}
+									}),
+								FCanExecuteAction::CreateWeakLambda(Context, [Context]() { return Context->bCanBeModified; }));
+
+							const TAttribute<FText> ToolTipOverride = Context->bCanBeModified ? TAttribute<FText>() : LOCTEXT("ReadOnlyAssetWarning", "The selected asset(s) are read-only and cannot be edited.");
+							Section.AddMenuEntryWithCommandList(
+								FWidgetPreviewCommands::Get().OpenEditor,
+								CommandListToBind,
+								TAttribute<FText>(),
+								ToolTipOverride,
+								FSlateIcon(FWidgetPreviewStyle::Get().GetStyleSetName(), "WidgetPreview.OpenEditor"));
+						}
+					}
+				}));
+		}
+	}
+}
+
+#undef LOCTEXT_NAMESPACE
+	
+IMPLEMENT_MODULE(UE::UMGWidgetPreview::Private::FUMGWidgetPreviewModule, UMGWidgetPreview)
diff --git a/Engine/Plugins/Editor/UMGWidgetPreview/Source/UMGWidgetPreview/Private/UMGWidgetPreviewModule.h b/Engine/Plugins/Editor/UMGWidgetPreview/Source/UMGWidgetPreview/Private/UMGWidgetPreviewModule.h
new file mode 100644
--- /dev/null
+++ b/Engine/Plugins/Editor/UMGWidgetPreview/Source/UMGWidgetPreview/Private/UMGWidgetPreviewModule.h
@@ -0,0 +1,21 @@
+﻿// Copyright Epic Games, Inc. All Rights Reserved.
+
+#pragma once
+
+#include "Modules/ModuleInterface.h"
+
+namespace UE::UMGWidgetPreview::Private
+{
+	class FUMGWidgetPreviewModule
+		: public IModuleInterface
+	{
+	public:
+		//~ Begin IModuleInterface
+		virtual void StartupModule() override;
+		virtual void ShutdownModule() override;
+		//~ End IModuleInterface
+
+	private:
+		void RegisterMenus();
+	};
+}
diff --git a/Engine/Plugins/Editor/UMGWidgetPreview/Source/UMGWidgetPreview/Private/WidgetPreview.cpp b/Engine/Plugins/Editor/UMGWidgetPreview/Source/UMGWidgetPreview/Private/WidgetPreview.cpp
new file mode 100644
--- /dev/null
+++ b/Engine/Plugins/Editor/UMGWidgetPreview/Source/UMGWidgetPreview/Private/WidgetPreview.cpp
@@ -0,0 +1,326 @@
+// Copyright Epic Games, Inc. All Rights Reserved.
+
+#include "WidgetPreview.h"
+
+#include "Blueprint/UserWidget.h"
+#include "WidgetBlueprint.h"
+#include "WidgetPreviewTypesPrivate.h"
+
+const TArray<FName>& UWidgetPreview::GetLayoutSlotNames()
+{
+	if (!SlotNameCache.IsEmpty())
+	{
+		return SlotNameCache;
+	}
+
+	if (!LayoutWidgetType.IsNull())
+	{
+		if (INamedSlotInterface* WidgetWithSlots = Cast<INamedSlotInterface>(LayoutWidgetType.LoadSynchronous()->GetDefaultObject<UUserWidget>()))
+		{
+			TArray<FName> SlotNames;
+			WidgetWithSlots->GetSlotNames(SlotNames);
+
+			SlotNameCache = SlotNames;
+			return SlotNameCache;
+		}
+	}
+
+	static TArray<FName> Empty;
+	return Empty;
+}
+
+UUserWidget* UWidgetPreview::GetOrCreateWidgetInstance(UWorld* InWorld, const bool bInForceRecreate)
+{
+	if (bInForceRecreate)
+	{
+		WidgetInstance = nullptr;
+	}
+	else if (WidgetInstance)
+	{
+		return WidgetInstance;
+	}
+
+	if (!ensure(InWorld))
+	{
+		return nullptr;
+	}
+
+	TArray<const UUserWidget*> UnsupportedWidgets;
+	if (!CanCallInitializedWithoutPlayerContext(true, UnsupportedWidgets))
+	{
+		// No need to log, this is an expected outcome
+		return nullptr;
+	}
+
+	if (UUserWidget* Widget = GetWidget())
+	{
+		auto MakeWidget = [InWorld](UClass* InClass) -> UUserWidget*
+		{
+			UUserWidget* NewWidget = NewObject<UUserWidget>(InWorld, InClass);
+			NewWidget->ClearFlags(RF_Transactional);
+			return NewWidget;
+		};
+
+		// @todo: always wrap with a container UWidget to account for warning here: void UWidget::RemoveFromParent()
+
+		WidgetInstance = MakeWidget(Widget->GetClass());
+
+		if (!LayoutWidgetType.IsNull() && !SlotWidgets.IsEmpty())
+		{
+			TArray<FName> ValidSlotNames;
+			WidgetInstance->GetSlotNames(ValidSlotNames);
+
+			for (const TPair<FName, TSoftClassPtr<UUserWidget>>& SlotWidget : SlotWidgets)
+			{
+				if (!SlotWidget.Value.IsNull()
+					&& ValidSlotNames.Contains(SlotWidget.Key))
+				{
+					WidgetInstance->SetContentForSlot(SlotWidget.Key, MakeWidget(SlotWidget.Value.LoadSynchronous()));
+				}
+			}
+		}
+
+		if (ULocalPlayer* LocalPlayer = InWorld->GetFirstLocalPlayerFromController())
+		{
+			WidgetInstance->SetPlayerContext(LocalPlayer);
+		}
+
+		OnWidgetChanged().Broadcast(EWidgetPreviewWidgetChangeType::Reinstanced);
+
+		return WidgetInstance;
+	}
+
+	return nullptr;
+}
+
+UUserWidget* UWidgetPreview::GetWidget() const
+{
+	// If the LayoutWidget is in use, return it (as the root widget).
+	if (!LayoutWidgetType.IsNull())
+	{
+		UUserWidget* LayoutWidgetCDO = LayoutWidgetType.LoadSynchronous()->GetDefaultObject<UUserWidget>();
+		if (INamedSlotInterface* WidgetWithSlots = Cast<INamedSlotInterface>(LayoutWidgetCDO))
+		{
+			TArray<FName> SlotNames;
+			WidgetWithSlots->GetSlotNames(SlotNames);
+			if (!SlotNames.IsEmpty())
+			{
+				return LayoutWidgetCDO;
+			}
+		}
+	}
+
+	if (!WidgetType.IsNull())
+	{
+		return WidgetType.LoadSynchronous()->GetDefaultObject<UUserWidget>();
+	}
+
+	return nullptr;
+}
+
+const UUserWidget* UWidgetPreview::GetWidgetForSlot(const FName InSlotName) const
+{
+	if (const TSoftClassPtr<UUserWidget>* WidgetInSlot = SlotWidgets.Find(InSlotName))
+	{
+		if (!WidgetInSlot->IsNull())
+		{
+			return WidgetInSlot->LoadSynchronous()->GetDefaultObject<UUserWidget>();	
+		}
+
+		UE_LOG(LogTemp, Warning, TEXT("Slot %s has invalid widget."), *InSlotName.ToString());
+	}
+
+	return nullptr;
+}
+
+bool UWidgetPreview::CanCallInitializedWithoutPlayerContext(const bool bInRecursive, TArray<const UUserWidget*>& OutFailedWidgets)
+{
+	if (!LayoutWidgetType.IsNull())
+	{
+		const UUserWidget* LayoutWidgetCDO = LayoutWidgetType.LoadSynchronous()->GetDefaultObject<UUserWidget>();
+		return CanCallInitializedWithoutPlayerContextOnWidget(LayoutWidgetCDO, bInRecursive, OutFailedWidgets);
+	}
+
+	if (!WidgetType.IsNull())
+	{
+		const UUserWidget* WidgetCDO = WidgetType.LoadSynchronous()->GetDefaultObject<UUserWidget>();
+		return CanCallInitializedWithoutPlayerContextOnWidget(WidgetCDO, bInRecursive, OutFailedWidgets);
+	}
+
+	// In case there are no widgets to display, we want to return true
+	return true;
+}
+
+bool UWidgetPreview::CanCallInitializedWithoutPlayerContextOnWidget(
+	const UUserWidget* InUserWidget,
+	const bool bInRecursive,
+	TArray<const UUserWidget*>& OutFailedWidgets)
+{
+	TFunction<bool(const UWidget* InWidgetGeneratedClass)> CanCallInitializedWithoutPlayerContextInternal;
+	CanCallInitializedWithoutPlayerContextInternal = [CanCallInitializedWithoutPlayerContextInternal, bInRecursive, &OutFailedWidgets](const UWidget* InWidget)
+	{
+		// In case there are no widgets to display, we want to return true
+		bool bResultInternal = true;
+
+		if (const UUserWidget* AsUserWidget = Cast<UUserWidget>(InWidget))
+		{
+			UE::UMGWidgetPreview::Private::FWidgetTypeTuple WidgetTuple(AsUserWidget);
+			if (WidgetTuple.BlueprintGeneratedClass)
+			{
+				bResultInternal = WidgetTuple.BlueprintGeneratedClass->bCanCallInitializedWithoutPlayerContext;
+				if (!bResultInternal)
+				{
+					OutFailedWidgets.Emplace(WidgetTuple.ClassDefaultObject);
+				}
+			}
+		}
+
+		if (bInRecursive)
+		{
+			if (const INamedSlotInterface* WidgetWithSlots = Cast<INamedSlotInterface>(InWidget))
+			{
+				TArray<FName> SlotNames;
+				WidgetWithSlots->GetSlotNames(SlotNames);
+				if (!SlotNames.IsEmpty())
+				{
+					for (const FName SlotName : SlotNames)
+					{
+						if (const UWidget* SlotWidget = Cast<UWidget>(WidgetWithSlots->GetContentForSlot(SlotName)))
+						{
+							bResultInternal = bResultInternal && CanCallInitializedWithoutPlayerContextInternal(SlotWidget);
+						}
+					}
+				}
+			}
+		}
+
+		return bResultInternal;
+	};
+
+	return CanCallInitializedWithoutPlayerContextInternal(InUserWidget);
+}
+
+const TSoftClassPtr<UUserWidget>& UWidgetPreview::GetWidgetType() const
+{
+	return WidgetType;
+}
+
+void UWidgetPreview::SetWidgetType(const TSoftClassPtr<UUserWidget>& InWidget)
+{
+	if (WidgetType != InWidget)
+	{
+		WidgetType = InWidget;
+		WidgetInstance = nullptr;
+		UpdateWidgets();
+
+		OnWidgetChanged().Broadcast(EWidgetPreviewWidgetChangeType::Assignment);
+	}
+}
+
+const TSoftClassPtr<UUserWidget>& UWidgetPreview::GetLayoutWidgetType() const
+{
+	return LayoutWidgetType;
+}
+
+void UWidgetPreview::SetLayoutWidgetType(const TSoftClassPtr<UUserWidget>& InWidget)
+{
+	if (LayoutWidgetType != InWidget)
+	{
+		LayoutWidgetType = InWidget;
+		WidgetInstance = nullptr;
+		UpdateWidgets();
+
+		OnWidgetChanged().Broadcast(EWidgetPreviewWidgetChangeType::Assignment);
+	}
+}
+
+const TMap<FName, TSoftClassPtr<UUserWidget>>& UWidgetPreview::GetSlotWidgets() const
+{
+	return SlotWidgets;
+}
+
+void UWidgetPreview::SetSlotWidgets(const TMap<FName, TSoftClassPtr<UUserWidget>>& InWidgets)
+{
+	if (!SlotWidgets.OrderIndependentCompareEqual(InWidgets))
+	{
+		SlotWidgets = InWidgets;
+		WidgetInstance = nullptr;
+		UpdateWidgets();
+
+		OnWidgetChanged().Broadcast(EWidgetPreviewWidgetChangeType::Assignment);
+	}
+}
+
+void UWidgetPreview::PostLoad()
+{
+	UObject::PostLoad();
+
+	UpdateWidgets();
+}
+
+void UWidgetPreview::PostEditChangeProperty(FPropertyChangedEvent& PropertyChangedEvent)
+{
+	Super::PostEditChangeProperty(PropertyChangedEvent);
+
+	FName PropertyName = PropertyChangedEvent.GetPropertyName();
+	if (PropertyName == GET_MEMBER_NAME_CHECKED(UWidgetPreview, WidgetType)
+		|| PropertyName == GET_MEMBER_NAME_CHECKED(UWidgetPreview, LayoutWidgetType)
+		|| PropertyName == GET_MEMBER_NAME_CHECKED(UWidgetPreview, SlotWidgets)
+		|| PropertyName.IsNone()) // None can be an Undo operation
+	{
+		WidgetInstance = nullptr;
+		UpdateWidgets();
+		OnWidgetChanged().Broadcast(EWidgetPreviewWidgetChangeType::Assignment);
+	}
+}
+
+void UWidgetPreview::OnWidgetBlueprintChanged(UBlueprint* InBlueprint)
+{
+	OnWidgetChanged().Broadcast(EWidgetPreviewWidgetChangeType::Structure);
+}
+
+void UWidgetPreview::UpdateWidgets()
+{
+	auto AddOnChanged = [this](const UUserWidget* InUserWidget)
+	{
+		const UE::UMGWidgetPreview::Private::FWidgetTypeTuple WidgetTuple(InUserWidget);
+		if (UWidgetBlueprint* BP = WidgetTuple.Blueprint)
+		{
+			BP->OnChanged().AddUObject(this, &UWidgetPreview::OnWidgetBlueprintChanged);
+		}
+	};
+
+	if (!LayoutWidgetType.IsNull())
+	{
+		SlotNameCache.Reset();
+
+		AddOnChanged(LayoutWidgetType.LoadSynchronous()->GetDefaultObject<UUserWidget>());
+
+		for (const TPair<FName, TSoftClassPtr<UUserWidget>>& SlotWidget : SlotWidgets)
+		{
+			if (!SlotWidget.Value.IsNull())
+			{
+				AddOnChanged(SlotWidget.Value.LoadSynchronous()->GetDefaultObject<UUserWidget>());
+			}
+		}
+	}
+
+	if (!WidgetType.IsNull())
+	{
+		AddOnChanged(WidgetType.LoadSynchronous()->GetDefaultObject<UUserWidget>());
+	}
+}
+
+TArray<FName> UWidgetPreview::GetAvailableLayoutSlotNames()
+{
+	const TArray<FName> AllSlotNames = GetLayoutSlotNames();
+
+	TArray<FName> UsedSlotNamesArray;
+	SlotWidgets.GenerateKeyArray(UsedSlotNamesArray);
+
+	const TSet<FName> UsedSlotNames(UsedSlotNamesArray);
+
+	return TSet<FName>(AllSlotNames)
+		.Difference(UsedSlotNames)
+		.Array();
+}
diff --git a/Engine/Plugins/Editor/UMGWidgetPreview/Source/UMGWidgetPreview/Private/WidgetPreview.h b/Engine/Plugins/Editor/UMGWidgetPreview/Source/UMGWidgetPreview/Private/WidgetPreview.h
new file mode 100644
--- /dev/null
+++ b/Engine/Plugins/Editor/UMGWidgetPreview/Source/UMGWidgetPreview/Private/WidgetPreview.h
@@ -0,0 +1,96 @@
+// Copyright Epic Games, Inc. All Rights Reserved.
+
+#pragma once
+
+#include "UObject/Object.h"
+
+#include "WidgetPreview.generated.h"
+
+class UWidgetBlueprint;
+struct FImage;
+class UWidget;
+class UUserWidget;
+
+enum class EWidgetPreviewWidgetChangeType
+{
+	Assignment = 0,
+	Reinstanced = 1,
+	Structure = 2,
+	ChildReference = 3,
+};
+
+// @todo: allow nested previews
+
+UCLASS(BlueprintType, NotBlueprintable, AutoExpandCategories = "Widgets")
+class UMGWIDGETPREVIEW_API UWidgetPreview
+	: public UObject
+{
+	GENERATED_BODY()
+
+public:
+	UFUNCTION(BlueprintCallable, Category = "Layout")
+	const TArray<FName>& GetLayoutSlotNames();
+
+	/** Returns or builds and returns an instance of the root widget for previewing. Can be used to trigger a rebuild. */
+	[[maybe_unused]] UUserWidget* GetOrCreateWidgetInstance(UWorld* InWorld, const bool bInForceRecreate = false);
+
+	UUserWidget* GetWidget() const;
+	const UUserWidget* GetWidgetForSlot(const FName InSlotName) const;
+
+	/** Convenience function to check that all utilized widgets have bCanCallInitializedWithoutPlayerContext set to true, and reports any that don't. */
+	bool CanCallInitializedWithoutPlayerContext(const bool bInRecursive, TArray<const UUserWidget*>& OutFailedWidgets);
+
+	// @todo: move to utility func somewhere else?
+	/** Convenience function to check that the provided widget (and it's children) has bCanCallInitializedWithoutPlayerContext set to true, and reports any that don't. */
+	static bool CanCallInitializedWithoutPlayerContextOnWidget(const UUserWidget* InUserWidget, const bool bInRecursive, TArray<const UUserWidget*>& OutFailedWidgets);
+
+public:
+	using FOnWidgetChanged = TMulticastDelegate<void(const EWidgetPreviewWidgetChangeType)>;
+
+	FOnWidgetChanged& OnWidgetChanged() { return OnWidgetChangedDelegate; }
+	
+	const TSoftClassPtr<UUserWidget>& GetWidgetType() const;
+	void SetWidgetType(const TSoftClassPtr<UUserWidget>& InWidget);
+	
+	const TSoftClassPtr<UUserWidget>& GetLayoutWidgetType() const;
+	void SetLayoutWidgetType(const TSoftClassPtr<UUserWidget>& InWidget);
+
+	const TMap<FName, TSoftClassPtr<UUserWidget>>& GetSlotWidgets() const;
+	void SetSlotWidgets(const TMap<FName, TSoftClassPtr<UUserWidget>>& InWidgets);
+
+protected:
+	virtual void PostLoad() override;
+
+	virtual void PostEditChangeProperty(FPropertyChangedEvent& PropertyChangedEvent) override;
+
+	void OnWidgetBlueprintChanged(UBlueprint* InBlueprint);
+
+	/** Misc. functionality to perform after a widget assignment is changed. */
+	void UpdateWidgets();
+
+	/** Returns slot names not already occupied in SlotWidgets. */
+	UFUNCTION(BlueprintCallable, Category = "Layout")
+	TArray<FName> GetAvailableLayoutSlotNames();
+
+private:
+	/** Widget to use. Will be overlaid on the Layout Widget (if specified). */
+	UPROPERTY(EditAnywhere, BlueprintReadWrite, Getter, Setter, Category = "Widget", DisplayName = "Widget", meta = (AllowPrivateAccess = "true"))
+	TSoftClassPtr<UUserWidget> WidgetType;
+
+	/** Optional widget to use for layout, containing named slots. */
+	UPROPERTY(EditAnywhere, BlueprintReadWrite, Getter, Setter, Category = "Layout", DisplayName = "Layout Widget", meta = (AllowPrivateAccess = "true"))
+	TSoftClassPtr<UUserWidget> LayoutWidgetType;
+
+	/** Widget per-slot. */
+	UPROPERTY(EditAnywhere, BlueprintReadWrite, Getter, Setter, Category = "Layout", meta = (AllowPrivateAccess = "true", EditCondition = "LayoutWidgetType", GetKeyOptions = "GetAvailableLayoutSlotNames"))
+	TMap<FName, TSoftClassPtr<UUserWidget>> SlotWidgets;
+
+	UPROPERTY(Transient)
+	TObjectPtr<UUserWidget> WidgetInstance;
+
+	/** Slot names available in LayoutWidgetType. */
+	UPROPERTY(Transient)
+	TArray<FName> SlotNameCache;
+
+	FOnWidgetChanged OnWidgetChangedDelegate;
+};
diff --git a/Engine/Plugins/Editor/UMGWidgetPreview/Source/UMGWidgetPreview/Private/WidgetPreviewCommands.cpp b/Engine/Plugins/Editor/UMGWidgetPreview/Source/UMGWidgetPreview/Private/WidgetPreviewCommands.cpp
new file mode 100644
--- /dev/null
+++ b/Engine/Plugins/Editor/UMGWidgetPreview/Source/UMGWidgetPreview/Private/WidgetPreviewCommands.cpp
@@ -0,0 +1,27 @@
+// Copyright Epic Games, Inc. All Rights Reserved.
+
+#include "WidgetPreviewCommands.h"
+
+#include "InputCoreTypes.h"
+#include "WidgetPreviewStyle.h"
+
+#define LOCTEXT_NAMESPACE "FWidgetPreviewCommands"
+
+namespace UE::UMGWidgetPreview::Private
+{
+	FWidgetPreviewCommands::FWidgetPreviewCommands()
+		: TCommands<FWidgetPreviewCommands>("WidgetPreviewEditor",
+			LOCTEXT("ContextDescription", "Widget Preview Editor"),
+			NAME_None,
+			FWidgetPreviewStyle::Get().GetStyleSetName())
+	{
+	}
+
+	void FWidgetPreviewCommands::RegisterCommands()
+	{
+		UI_COMMAND(OpenEditor, "Widget Preview", "Opens the Widget Preview window.", EUserInterfaceActionType::Button, FInputChord());
+		UI_COMMAND(ResetPreview, "Reset", "Resets the current Widget Preview.", EUserInterfaceActionType::Button, FInputChord(EModifierKey::Control, EKeys::R));
+	}
+}
+
+#undef LOCTEXT_NAMESPACE
diff --git a/Engine/Plugins/Editor/UMGWidgetPreview/Source/UMGWidgetPreview/Private/WidgetPreviewCommands.h b/Engine/Plugins/Editor/UMGWidgetPreview/Source/UMGWidgetPreview/Private/WidgetPreviewCommands.h
new file mode 100644
--- /dev/null
+++ b/Engine/Plugins/Editor/UMGWidgetPreview/Source/UMGWidgetPreview/Private/WidgetPreviewCommands.h
@@ -0,0 +1,25 @@
+// Copyright Epic Games, Inc. All Rights Reserved.
+
+#pragma once
+
+#include "Framework/Commands/Commands.h"
+
+namespace UE::UMGWidgetPreview::Private
+{
+	class FWidgetPreviewCommands
+		: public TCommands<FWidgetPreviewCommands>
+	{
+	public:
+		FWidgetPreviewCommands();
+
+		//~ Begin TCommands
+		virtual void RegisterCommands() override;
+		//~ End TCommands
+
+	public:
+		TSharedPtr<FUICommandInfo> OpenEditor;
+
+		/** Resets the state of the given preview. */
+		TSharedPtr<FUICommandInfo> ResetPreview;
+	};
+}
diff --git a/Engine/Plugins/Editor/UMGWidgetPreview/Source/UMGWidgetPreview/Private/WidgetPreviewEditor.cpp b/Engine/Plugins/Editor/UMGWidgetPreview/Source/UMGWidgetPreview/Private/WidgetPreviewEditor.cpp
new file mode 100644
--- /dev/null
+++ b/Engine/Plugins/Editor/UMGWidgetPreview/Source/UMGWidgetPreview/Private/WidgetPreviewEditor.cpp
@@ -0,0 +1,106 @@
+// Copyright Epic Games, Inc. All Rights Reserved.
+
+#include "WidgetPreviewEditor.h"
+
+#include "Algo/AllOf.h"
+#include "AssetToolsModule.h"
+#include "Blueprint/UserWidget.h"
+#include "PackageTools.h"
+#include "WidgetBlueprint.h"
+#include "WidgetPreview.h"
+#include "WidgetPreviewFactory.h"
+#include "WidgetPreviewToolkit.h"
+
+void UWidgetPreviewEditor::Initialize(const TObjectPtr<UWidgetPreview>& InWidgetPreview)
+{
+	WidgetPreview = InWidgetPreview;
+
+	UAssetEditor::Initialize();
+}
+
+void UWidgetPreviewEditor::GetObjectsToEdit(TArray<UObject*>& InObjectsToEdit)
+{
+	InObjectsToEdit.Emplace(WidgetPreview);
+}
+
+TSharedPtr<FBaseAssetToolkit> UWidgetPreviewEditor::CreateToolkit()
+{
+	return MakeShared<UE::UMGWidgetPreview::Private::FWidgetPreviewToolkit>(this);
+}
+
+void UWidgetPreviewEditor::FocusWindow(UObject* ObjectToFocusOn)
+{
+	if (ToolkitInstance)
+	{
+		ToolkitInstance->FocusWindow(ObjectToFocusOn);
+	}
+}
+
+UWidgetPreview* UWidgetPreviewEditor::GetObjectToEdit() const
+{
+	return WidgetPreview;
+}
+
+bool UWidgetPreviewEditor::AreObjectsValidTargets(const TArray<UObject*>& InObjects)
+{
+	if (InObjects.IsEmpty())
+	{
+		return false;
+	}
+
+	return Algo::AllOf(InObjects, [](const UObject* InObject)
+	{
+		return InObject->IsA<UWidgetBlueprint>();
+	});
+}
+
+bool UWidgetPreviewEditor::AreAssetsValidTargets(const TArray<FAssetData>& InAssets)
+{
+	if (InAssets.IsEmpty())
+	{
+		return false;
+	}
+
+	return Algo::AllOf(InAssets, [](const FAssetData& InAsset)
+	{
+		return InAsset.IsInstanceOf<UWidgetBlueprint>();
+	});
+}
+
+UWidgetPreview* UWidgetPreviewEditor::CreatePreviewForWidget(const UUserWidget* InUserWidget)
+{
+	// Create a new widget preview (@see: UEditorEngine::NewMap)
+	UWidgetPreviewFactory* Factory = NewObject<UWidgetPreviewFactory>();
+
+	FAssetToolsModule& AssetToolsModule = FModuleManager::GetModuleChecked<FAssetToolsModule>("AssetTools");
+
+	FString WidgetAssetPath;
+	FString WidgetAssetName;
+	InUserWidget->GetPackage()->GetName().Split(TEXT("/"), &WidgetAssetPath, &WidgetAssetName, ESearchCase::IgnoreCase, ESearchDir::FromEnd);
+
+	const FString DesiredAssetName = WidgetAssetName + TEXT("_Preview");
+	const FString DesiredPackagePath = UPackageTools::SanitizePackageName(WidgetAssetPath + TEXT("/") + DesiredAssetName);
+
+	FString DefaultSuffix;
+	FString AssetName;
+	FString PackageName;
+	AssetToolsModule.Get().CreateUniqueAssetName(DesiredPackagePath, DefaultSuffix, PackageName, AssetName);
+
+	UPackage* Package = CreatePackage(*PackageName);
+	Package->SetPackageFlags(PKG_NewlyCreated);
+
+	EObjectFlags Flags = RF_Public | RF_Standalone;
+	UWidgetPreview* NewWidgetPreview = CastChecked<UWidgetPreview>(
+		Factory->FactoryCreateNew(
+			UWidgetPreview::StaticClass(),
+			Package,
+			*AssetName,
+			Flags,
+			nullptr,
+			GWarn));
+
+	NewWidgetPreview->SetWidgetType(InUserWidget->GetClass());
+	NewWidgetPreview->MarkPackageDirty();
+
+	return NewWidgetPreview;
+}
diff --git a/Engine/Plugins/Editor/UMGWidgetPreview/Source/UMGWidgetPreview/Private/WidgetPreviewEditor.h b/Engine/Plugins/Editor/UMGWidgetPreview/Source/UMGWidgetPreview/Private/WidgetPreviewEditor.h
new file mode 100644
--- /dev/null
+++ b/Engine/Plugins/Editor/UMGWidgetPreview/Source/UMGWidgetPreview/Private/WidgetPreviewEditor.h
@@ -0,0 +1,49 @@
+// Copyright Epic Games, Inc. All Rights Reserved.
+
+#pragma once
+
+#include "Tools/UAssetEditor.h"
+
+#include "WidgetPreviewEditor.generated.h"
+
+class FBaseAssetToolkit;
+class UObject;
+class UUserWidget;
+class UWidget;
+class UWidgetPreview;
+
+/** Minimalistic container to spawn the toolkit. */
+UCLASS(Transient)
+class UWidgetPreviewEditor
+	: public UAssetEditor
+{
+	GENERATED_BODY()
+
+public:
+	void Initialize(const TObjectPtr<UWidgetPreview>& InWidgetPreview);
+
+	//~ Begin UAssetEditor
+	virtual void GetObjectsToEdit(TArray<UObject*>& InObjectsToEdit) override;
+	virtual TSharedPtr<FBaseAssetToolkit> CreateToolkit() override;
+	virtual void FocusWindow(UObject* ObjectToFocusOn) override;
+	//~ End UAssetEditor
+
+	UWidgetPreview* GetObjectToEdit() const;
+
+	/** Checks that all of the objects are valid targets for a Widget Preview session. */
+	static bool AreObjectsValidTargets(const TArray<UObject*>& InObjects);
+
+	/**
+	 * Checks that all of the assets are valid targets for an editor session. This
+	 * is preferable over AreObjectsValidTargets when we have FAssetData because it
+	 * allows us to avoid forcing a load of the underlying UObjects (for instance to
+	 * avoid triggering a load when right clicking an asset in the content browser).
+	 */
+	static bool AreAssetsValidTargets(const TArray<FAssetData>& InAssets);
+
+	/** Creates a new, unsaved Widget Preview asset for the given UUserWidget. */
+	static UWidgetPreview* CreatePreviewForWidget(const UUserWidget* InUserWidget);
+
+private:
+	TObjectPtr<UWidgetPreview> WidgetPreview = nullptr;
+};
diff --git a/Engine/Plugins/Editor/UMGWidgetPreview/Source/UMGWidgetPreview/Private/WidgetPreviewFactory.cpp b/Engine/Plugins/Editor/UMGWidgetPreview/Source/UMGWidgetPreview/Private/WidgetPreviewFactory.cpp
new file mode 100644
--- /dev/null
+++ b/Engine/Plugins/Editor/UMGWidgetPreview/Source/UMGWidgetPreview/Private/WidgetPreviewFactory.cpp
@@ -0,0 +1,21 @@
+// Copyright Epic Games, Inc. All Rights Reserved.
+
+#include "WidgetPreviewFactory.h"
+
+#include "WidgetPreview.h"
+
+UWidgetPreviewFactory::UWidgetPreviewFactory()
+{
+	SupportedClass = UWidgetPreview::StaticClass();
+	bCreateNew = true;
+	bEditAfterNew = true;
+}
+
+UObject* UWidgetPreviewFactory::FactoryCreateNew(UClass* InClass, UObject* InParent, FName InName, EObjectFlags Flags, UObject* Context, FFeedbackContext* Warn)
+{
+	check(InClass == UWidgetPreview::StaticClass());
+
+	UWidgetPreview* WidgetPreviewAsset = NewObject<UWidgetPreview>(InParent, InName, Flags);
+	
+	return WidgetPreviewAsset;
+}
diff --git a/Engine/Plugins/Editor/UMGWidgetPreview/Source/UMGWidgetPreview/Private/WidgetPreviewFactory.h b/Engine/Plugins/Editor/UMGWidgetPreview/Source/UMGWidgetPreview/Private/WidgetPreviewFactory.h
new file mode 100644
--- /dev/null
+++ b/Engine/Plugins/Editor/UMGWidgetPreview/Source/UMGWidgetPreview/Private/WidgetPreviewFactory.h
@@ -0,0 +1,20 @@
+// Copyright Epic Games, Inc. All Rights Reserved.
+
+#pragma once
+
+#include "Factories/Factory.h"
+
+#include "WidgetPreviewFactory.generated.h"
+
+UCLASS(HideCategories = Object, CollapseCategories, MinimalAPI)
+class UWidgetPreviewFactory : public UFactory
+{
+	GENERATED_BODY()
+
+public:
+	UWidgetPreviewFactory();
+	
+	//~ Begin UFactory
+	virtual UObject* FactoryCreateNew(UClass* InClass, UObject* InParent, FName InName, EObjectFlags Flags, UObject* Context, FFeedbackContext* Warn) override;
+	//~ End UFactory
+};
diff --git a/Engine/Plugins/Editor/UMGWidgetPreview/Source/UMGWidgetPreview/Private/WidgetPreviewLog.h b/Engine/Plugins/Editor/UMGWidgetPreview/Source/UMGWidgetPreview/Private/WidgetPreviewLog.h
new file mode 100644
--- /dev/null
+++ b/Engine/Plugins/Editor/UMGWidgetPreview/Source/UMGWidgetPreview/Private/WidgetPreviewLog.h
@@ -0,0 +1,13 @@
+// Copyright Epic Games, Inc. All Rights Reserved.
+
+#pragma once
+
+#include "Logging/LogMacros.h"
+
+
+namespace UE::UMGWidgetPreview::Private
+{
+	static const FLazyName MessageLogName = "WidgetPreviewLog";
+}
+
+DECLARE_LOG_CATEGORY_EXTERN(LogWidgetPreview, VeryVerbose, All);
diff --git a/Engine/Plugins/Editor/UMGWidgetPreview/Source/UMGWidgetPreview/Private/WidgetPreviewStyle.cpp b/Engine/Plugins/Editor/UMGWidgetPreview/Source/UMGWidgetPreview/Private/WidgetPreviewStyle.cpp
new file mode 100644
--- /dev/null
+++ b/Engine/Plugins/Editor/UMGWidgetPreview/Source/UMGWidgetPreview/Private/WidgetPreviewStyle.cpp
@@ -0,0 +1,50 @@
+// Copyright Epic Games, Inc. All Rights Reserved.
+
+#include "WidgetPreviewStyle.h"
+
+#include "Interfaces/IPluginManager.h"
+#include "Styling/AppStyle.h"
+#include "Styling/CoreStyle.h"
+#include "Styling/SlateStyleMacros.h"
+#include "Styling/SlateTypes.h"
+
+namespace UE::UMGWidgetPreview::Private
+{
+	FName FWidgetPreviewStyle::StyleName("WidgetPreview");
+
+	FWidgetPreviewStyle& FWidgetPreviewStyle::Get()
+	{
+		static FWidgetPreviewStyle Instance;
+		return Instance;
+	}
+
+	FWidgetPreviewStyle::FWidgetPreviewStyle()
+		: FSlateStyleSet(StyleName)
+	{
+		const FVector2f Icon16(16.f);
+		const FVector2f Icon20(20.f);
+		const FVector2f Icon22(22.f);
+		const FVector2f Icon25(25.f);
+		const FVector2f Icon32(32.f);
+		const FVector2f Icon40(40.f);
+
+		TSharedPtr<IPlugin> Plugin = IPluginManager::Get().FindPlugin(UE_PLUGIN_NAME);
+		check(Plugin.IsValid());
+
+		SetContentRoot(Plugin->GetBaseDir() / TEXT("Resources"));
+		SetCoreContentRoot(FPaths::EngineContentDir() / TEXT("Editor/Slate"));
+
+		Set("WidgetPreview.OpenEditor", new IMAGE_BRUSH_SVG("Icons/WidgetPreview_16", Icon16));
+		Set("WidgetPreview.Reset", new CORE_IMAGE_BRUSH_SVG("Starship/Common/Reset", Icon16));
+
+		// Custom Class Icons
+		Set("ClassIcon.WidgetPreview", new IMAGE_BRUSH_SVG(TEXT("Icons/WidgetPreview_16"), Icon16));
+
+		FSlateStyleRegistry::RegisterSlateStyle(*this);
+	}
+
+	FWidgetPreviewStyle::~FWidgetPreviewStyle()
+	{
+		FSlateStyleRegistry::UnRegisterSlateStyle(*this);
+	}
+}
diff --git a/Engine/Plugins/Editor/UMGWidgetPreview/Source/UMGWidgetPreview/Private/WidgetPreviewStyle.h b/Engine/Plugins/Editor/UMGWidgetPreview/Source/UMGWidgetPreview/Private/WidgetPreviewStyle.h
new file mode 100644
--- /dev/null
+++ b/Engine/Plugins/Editor/UMGWidgetPreview/Source/UMGWidgetPreview/Private/WidgetPreviewStyle.h
@@ -0,0 +1,22 @@
+// Copyright Epic Games, Inc. All Rights Reserved.
+
+#pragma once
+
+#include "Styling/SlateStyle.h"
+#include "Styling/SlateStyleRegistry.h"
+
+namespace UE::UMGWidgetPreview::Private
+{
+	class FWidgetPreviewStyle
+		: public FSlateStyleSet
+	{
+	public:
+		static FName StyleName;
+
+		static FWidgetPreviewStyle& Get();
+
+	private:
+		FWidgetPreviewStyle();
+		virtual ~FWidgetPreviewStyle() override;
+	};
+}
diff --git a/Engine/Plugins/Editor/UMGWidgetPreview/Source/UMGWidgetPreview/Private/WidgetPreviewToolkit.cpp b/Engine/Plugins/Editor/UMGWidgetPreview/Source/UMGWidgetPreview/Private/WidgetPreviewToolkit.cpp
new file mode 100644
--- /dev/null
+++ b/Engine/Plugins/Editor/UMGWidgetPreview/Source/UMGWidgetPreview/Private/WidgetPreviewToolkit.cpp
@@ -0,0 +1,641 @@
+// Copyright Epic Games, Inc. All Rights Reserved.
+
+#include "WidgetPreviewToolkit.h"
+
+#include "Brushes/SlateRoundedBoxBrush.h"
+#include "DataValidationFixers.h"
+#include "Framework/Application/SlateApplication.h"
+#include "Framework/Application/SlateUser.h"
+#include "Kismet2/KismetEditorUtilities.h"
+#include "MessageLogModule.h"
+#include "ObjectEditorUtils.h"
+#include "SAssetEditorViewport.h"
+#include "ScopedTransaction.h"
+#include "Styling/AppStyle.h"
+#include "WidgetBlueprintEditor.h"
+#include "WidgetPreview.h"
+#include "WidgetPreviewCommands.h"
+#include "WidgetPreviewEditor.h"
+#include "WidgetPreviewLog.h"
+#include "WidgetPreviewStyle.h"
+#include "WidgetPreviewTypesPrivate.h"
+#include "Widgets/Docking/SDockTab.h"
+#include "Widgets/SWidgetPreview.h"
+#include "Widgets/SWidgetPreviewDetails.h"
+#include "Widgets/SWidgetPreviewStatus.h"
+#include "Widgets/SWidgetPreviewViewport.h"
+
+#define LOCTEXT_NAMESPACE "WidgetPreviewToolkit"
+
+namespace UE::UMGWidgetPreview::Private
+{
+	const FLazyName FWidgetPreviewToolkit::PreviewSceneSettingsTabID(TEXT("WidgetPreviewToolkit_PreviewScene"));
+	const FLazyName FWidgetPreviewToolkit::MessageLogTabID(TEXT("WidgetPreviewToolkit_MessageLog"));
+
+	EFixApplicability FWidgetPreviewabilityFixer::GetApplicability(int32 FixIndex) const
+	{
+		if (const UUserWidget* UserWidget = WeakUserWidget.Get())
+		{
+			const FWidgetTypeTuple WidgetTuple(UserWidget);
+			if (WidgetTuple.BlueprintGeneratedClass)
+			{
+				return WidgetTuple.BlueprintGeneratedClass->bCanCallInitializedWithoutPlayerContext
+					? EFixApplicability::Applied
+					: EFixApplicability::CanBeApplied;
+			}
+		}
+
+		return EFixApplicability::DidNotApply;
+	}
+
+	FFixResult FWidgetPreviewabilityFixer::ApplyFix(int32 FixIndex)
+	{
+		// @todo: apply recursively (named slots, etc.)
+		if (const UUserWidget* UserWidget = WeakUserWidget.Get())
+		{
+			FWidgetTypeTuple WidgetTuple(UserWidget);
+			if (UWidgetBlueprint* WidgetBlueprint = WidgetTuple.Blueprint)
+			{
+				FScopedTransaction Transaction(LOCTEXT("FixWidgetBlueprint", "Fix Widget Blueprint"));
+
+				// Set flag
+				{
+					FObjectEditorUtils::SetPropertyValue(
+						WidgetBlueprint,
+						GET_MEMBER_NAME_CHECKED(UWidgetBlueprint, bCanCallInitializedWithoutPlayerContext),
+						true);
+				}
+
+				// Compile
+				{
+					FKismetEditorUtilities::CompileBlueprint(WidgetBlueprint, EBlueprintCompileOptions::SkipGarbageCollection);
+					WidgetBlueprint->PostEditChange();
+					WidgetBlueprint->MarkPackageDirty();
+				}
+
+				return FFixResult::Success();
+			}
+		}
+
+		return FFixResult::Failure(LOCTEXT("FixWidgetBlueprint_Failure", "Failed to fix UserWidget."));
+	}
+
+	TSharedRef<FWidgetPreviewabilityFixer> FWidgetPreviewabilityFixer::Create(const UUserWidget* InUserWidget)
+	{
+		TSharedRef<FWidgetPreviewabilityFixer> Fixer = MakeShared<FWidgetPreviewabilityFixer>();
+		Fixer->WeakUserWidget = InUserWidget;
+		return Fixer;
+	}
+
+	FWidgetPreviewScene::FWidgetPreviewScene(const TSharedRef<FWidgetPreviewToolkit>& InPreviewToolkit)
+		: WeakToolkit(InPreviewToolkit)
+		, PreviewScene(MakeShared<FAdvancedPreviewScene>(
+			FPreviewScene::ConstructionValues()
+			.AllowAudioPlayback(true)
+			.ShouldSimulatePhysics(true)
+			.SetEditor(false)))
+	{
+		PreviewScene->SetFloorVisibility(false);
+	}
+
+	void FWidgetPreviewScene::Tick(float DeltaTime)
+	{
+		if (const TSharedPtr<FWidgetPreviewToolkit> Toolkit = WeakToolkit.Pin())
+		{
+			if (!Toolkit->GetState()->CanTick()
+				|| GEditor->bIsSimulatingInEditor
+				|| GEditor->PlayWorld != nullptr)
+			{
+				return;
+			}
+
+			PreviewScene->GetWorld()->Tick(LEVELTICK_All, DeltaTime);
+		}
+	}
+
+	ETickableTickType FWidgetPreviewScene::GetTickableTickType() const
+	{
+		return ETickableTickType::Always;
+	}
+
+	TStatId FWidgetPreviewScene::GetStatId() const
+	{
+		RETURN_QUICK_DECLARE_CYCLE_STAT(FWidgetPreviewScene, STATGROUP_Tickables);
+	}
+
+	UWorld* FWidgetPreviewScene::GetWorld() const
+	{
+		return PreviewScene->GetWorld();
+	}
+
+	TSharedRef<FAdvancedPreviewScene> FWidgetPreviewScene::GetPreviewScene() const
+	{
+		return PreviewScene.ToSharedRef();
+	}
+
+	FName FWidgetPreviewToolkitStateBase::GetId() const
+	{
+		return Id;
+	}
+
+	const TSharedPtr<FTokenizedMessage>& FWidgetPreviewToolkitStateBase::GetStatusMessage() const
+	{
+		return StatusMessage;
+	}
+
+	bool FWidgetPreviewToolkitStateBase::CanTick() const
+	{
+		return bCanTick;
+	}
+
+	bool FWidgetPreviewToolkitStateBase::ShouldOverlayStatusMessage() const
+	{
+		return bShouldOverlayMessage;
+	}
+
+	void FWidgetPreviewToolkitStateBase::OnEnter(const FWidgetPreviewToolkitStateBase* InFromState)
+	{
+	}
+
+	void FWidgetPreviewToolkitStateBase::OnExit(const FWidgetPreviewToolkitStateBase* InToState)
+	{
+	}
+
+	FWidgetPreviewToolkitPausedState::FWidgetPreviewToolkitPausedState(): FWidgetPreviewToolkitStateBase(TEXT("Paused"))
+	{
+		StatusMessage = FTokenizedMessage::Create(EMessageSeverity::Info, LOCTEXT("WidgetPreviewToolkitPausedState_Message", "The preview is currently paused."));
+		bCanTick = false;
+		bShouldOverlayMessage = true;
+	}
+
+	FWidgetPreviewToolkitBackgroundState::FWidgetPreviewToolkitBackgroundState()
+	{
+		Id = TEXT("Background");
+		StatusMessage = FTokenizedMessage::Create(EMessageSeverity::Info, LOCTEXT("WidgetPreviewToolkitBackgroundState_Message", "The widget preview is paused while the window is in the background. Re-focus to unpause."));
+	}
+
+	FWidgetPreviewToolkitUnsupportedWidgetState::FWidgetPreviewToolkitUnsupportedWidgetState()
+	{
+		Id = TEXT("UnsupportedWidget");
+		ResetStatusMessage();
+	}
+
+	void FWidgetPreviewToolkitUnsupportedWidgetState::SetUnsupportedWidgets(const TArray<const UUserWidget*>& InWidgets)
+	{
+		UnsupportedWidgets.Reset();
+
+		Algo::Transform(
+			InWidgets,
+			UnsupportedWidgets,
+			[](const UUserWidget* InWidget){
+				return MakeWeakObjectPtr(InWidget);
+			});
+
+		// Reset message
+		ResetStatusMessage();
+
+		for (const TWeakObjectPtr<const UUserWidget>& WeakUnsupportedWidget : UnsupportedWidgets)
+		{
+			if (const UUserWidget* UnsupportedWidget = WeakUnsupportedWidget.Get())
+			{
+				const TSharedRef<FWidgetPreviewabilityFixer> WidgetFixer = FWidgetPreviewabilityFixer::Create(UnsupportedWidget);
+
+				StatusMessage->AddToken(FAssetNameToken::Create(UnsupportedWidget->GetPackage()->GetName()));
+				StatusMessage->AddToken(WidgetFixer->CreateToken(LOCTEXT("FixUnsupportedWidget", "Fix")));
+			}
+		}
+	}
+
+	void FWidgetPreviewToolkitUnsupportedWidgetState::ResetStatusMessage()
+	{
+		StatusMessage = FTokenizedMessage::Create(EMessageSeverity::Error, LOCTEXT("WidgetPreviewToolkitUnsupportedWidgetState_Message", "One or more referenced widgets isn't supported (\"Can Call Initialized Without Player Context\" might be disabled)."));
+	}
+
+	FWidgetPreviewToolkitRunningState::FWidgetPreviewToolkitRunningState(): FWidgetPreviewToolkitStateBase(TEXT("Running"))
+	{
+		StatusMessage = FTokenizedMessage::Create(EMessageSeverity::Info, LOCTEXT("WidgetPreviewToolkitRunningState_Message", "The preview is running!"));
+		bCanTick = true;
+		bShouldOverlayMessage = false;
+	}
+
+	FWidgetPreviewToolkit::FWidgetPreviewToolkit(UWidgetPreviewEditor* InOwningEditor)
+		: FBaseAssetToolkit(InOwningEditor)
+		, Preview(InOwningEditor->GetObjectToEdit())
+	{
+		StandaloneDefaultLayout = FTabManager::NewLayout("WidgetPreviewEditor_Layout_v1")
+			->AddArea
+			(
+				FTabManager::NewPrimaryArea()
+				->SetOrientation(Orient_Horizontal)
+				->Split
+				(
+					FTabManager::NewSplitter()
+					->SetOrientation(Orient_Vertical)
+					->SetSizeCoefficient(0.85f)
+					->Split
+					(
+						FTabManager::NewStack()
+						->SetSizeCoefficient(0.85f)
+						->AddTab(ViewportTabID, ETabState::OpenedTab)
+						->SetHideTabWell(true)
+					)
+					->Split
+					(
+					FTabManager::NewStack()
+						->SetSizeCoefficient(0.15f)
+						->AddTab(MessageLogTabID, ETabState::OpenedTab)
+					)
+				)
+				->Split
+				(
+					FTabManager::NewSplitter()
+					->SetOrientation(Orient_Vertical)
+					->SetSizeCoefficient(0.15f)
+					->Split
+					(
+						FTabManager::NewStack()
+						->AddTab(DetailsTabID, ETabState::OpenedTab)
+						->AddTab(PreviewSceneSettingsTabID, ETabState::OpenedTab)
+						->SetForegroundTab(DetailsTabID)
+					)
+				)
+			);
+	}
+
+	FWidgetPreviewToolkit::~FWidgetPreviewToolkit()
+	{
+		if (Preview)
+		{
+			Preview->OnWidgetChanged().Remove(OnWidgetChangedHandle);
+		}
+
+		// Ensure remaining references to the update state stop ticking
+		SetState(&PausedState);
+
+		if (FSlateApplication::IsInitialized())
+		{
+			FSlateApplication::Get().OnFocusChanging().Remove(OnFocusChangingHandle);
+		}
+	}
+
+	void FWidgetPreviewToolkit::CreateWidgets()
+	{
+		FBaseAssetToolkit::CreateWidgets();
+
+		FMessageLogModule& MessageLogModule = FModuleManager::LoadModuleChecked<FMessageLogModule>("MessageLog");
+		MessageLogListing = MessageLogModule.GetLogListing(MessageLogName);
+		MessageLogWidget = MessageLogModule.CreateLogListingWidget(MessageLogListing.ToSharedRef());
+
+		FAdvancedPreviewSceneModule& AdvancedPreviewSceneModule = FModuleManager::LoadModuleChecked<FAdvancedPreviewSceneModule>("AdvancedPreviewScene");
+
+		const TSharedRef<FWidgetPreviewToolkit>& Self = SharedThis(this);
+		PreviewScene = MakeShared<FWidgetPreviewScene>(Self);
+
+		TArray<FAdvancedPreviewSceneModule::FDetailDelegates> Delegates;
+		Delegates.Add({ OnPreviewSceneChangedDelegate });
+		PreviewSettingsWidget = AdvancedPreviewSceneModule.CreateAdvancedPreviewSceneSettingsWidget(
+			PreviewScene->GetPreviewScene(),
+			nullptr,
+			TArray<FAdvancedPreviewSceneModule::FDetailCustomizationInfo>(),  TArray<FAdvancedPreviewSceneModule::FPropertyTypeCustomizationInfo>(),
+			Delegates);
+	}
+
+	void FWidgetPreviewToolkit::RegisterToolbar()
+	{
+		FName ParentName;
+		const FName MenuName = GetToolMenuToolbarName(ParentName);
+		UToolMenus* ToolMenus = UToolMenus::Get();
+		if (!ToolMenus->IsMenuRegistered(MenuName))
+		{
+			FToolMenuOwnerScoped ToolMenuOwnerScope(this);
+
+			UToolMenu* ToolbarMenu = ToolMenus->RegisterMenu(MenuName, ParentName, EMultiBoxType::ToolBar);
+
+			FToolMenuInsert InsertAfterAssetSection("Asset", EToolMenuInsertType::After);
+			const FWidgetPreviewCommands& Commands = FWidgetPreviewCommands::Get();
+
+			FToolMenuSection& PreviewSection = ToolbarMenu->FindOrAddSection(
+				"Preview",
+				{},
+				InsertAfterAssetSection);
+
+			PreviewSection.AddEntry(
+				FToolMenuEntry::InitToolBarButton(Commands.ResetPreview,
+					FText::GetEmpty(),
+					{},
+					FSlateIcon(FWidgetPreviewStyle::Get().GetStyleSetName(), "WidgetPreview.Reset")));
+		}
+	}
+
+	void FWidgetPreviewToolkit::RegisterTabSpawners(const TSharedRef<FTabManager>& InTabManager)
+	{
+		FBaseAssetToolkit::RegisterTabSpawners(InTabManager);
+
+		ensure(AssetEditorTabsCategory.IsValid());
+
+		const TSharedRef<FWorkspaceItem> AssetEditorTabsCategoryRef = AssetEditorTabsCategory.ToSharedRef();
+
+		InTabManager->RegisterTabSpawner(
+			PreviewSceneSettingsTabID,
+			FOnSpawnTab::CreateSP(this, &FWidgetPreviewToolkit::SpawnTab_PreviewSceneSettings))
+			.SetDisplayName(LOCTEXT("PreviewSceneTab", "Preview Scene Settings"))
+			.SetGroup(AssetEditorTabsCategoryRef)
+			.SetIcon(FSlateIcon(FAppStyle::GetAppStyleSetName(), "LevelEditor.Tabs.Details"))
+			.SetReadOnlyBehavior(ETabReadOnlyBehavior::Disabled);
+
+		InTabManager->RegisterTabSpawner(
+			MessageLogTabID,
+			FOnSpawnTab::CreateSP(this, &FWidgetPreviewToolkit::SpawnTab_MessageLog))
+			.SetDisplayName(LOCTEXT("MessageLogTab", "Message Log"))
+			.SetGroup(AssetEditorTabsCategoryRef)
+			.SetIcon(FSlateIcon(FAppStyle::GetAppStyleSetName(), "MessageLog.TabIcon"));
+	}
+
+	void FWidgetPreviewToolkit::UnregisterTabSpawners(const TSharedRef<FTabManager>& InTabManager)
+	{
+		FBaseAssetToolkit::UnregisterTabSpawners(InTabManager);
+
+		InTabManager->UnregisterTabSpawner(PreviewSceneSettingsTabID);
+		InTabManager->UnregisterTabSpawner(MessageLogTabID);
+	}
+
+	void FWidgetPreviewToolkit::PostInitAssetEditor()
+	{
+		OnWidgetChangedHandle = Preview->OnWidgetChanged().AddSP(this, &FWidgetPreviewToolkit::OnWidgetChanged);
+
+		if (FSlateApplication::IsInitialized())
+		{
+			OnFocusChangingHandle = FSlateApplication::Get().OnFocusChanging().AddSP(this, &FWidgetPreviewToolkit::OnFocusChanging);
+		}
+
+		// Bind Commands
+		{
+			const FWidgetPreviewCommands& Commands = FWidgetPreviewCommands::Get();
+
+			ToolkitCommands->MapAction(
+				Commands.ResetPreview,
+				FExecuteAction::CreateSP(this, &FWidgetPreviewToolkit::ResetPreview));
+		}
+
+		ResolveState();
+	}
+
+	FText FWidgetPreviewToolkit::GetToolkitName() const
+	{
+		const TArray<UObject*>* Objects = GetObjectsCurrentlyBeingEdited();
+
+		// Singular
+		if (Objects->Num() == 1)
+		{
+			return FText::Format(LOCTEXT("WidgetPreviewTabNameWithObject", "Widget Preview: {0}"), GetLabelForObject((*Objects)[0]));
+		}
+
+		// Plural
+		return LOCTEXT("WidgetPreviewTabNameWithObjects", "Widget Preview: (Multiple)");
+	}
+
+	FName FWidgetPreviewToolkit::GetToolkitFName() const
+	{
+		return FName(FString::Printf(TEXT("WidgetPreview%p"), this));
+	}
+
+	FText FWidgetPreviewToolkit::GetBaseToolkitName() const
+	{
+		return LOCTEXT("WidgetPreviewToolkitName", "Widget Preview");
+	}
+
+	FString FWidgetPreviewToolkit::GetWorldCentricTabPrefix() const
+	{
+		return LOCTEXT("WorldCentricTabPrefix", "WidgetPreviw").ToString();
+	}
+
+	FLinearColor FWidgetPreviewToolkit::GetWorldCentricTabColorScale() const
+	{
+		return FLinearColor(0.3f, 0.2f, 0.5f, 0.5f);
+	}
+
+	void FWidgetPreviewToolkit::AddReferencedObjects(FReferenceCollector& Collector)
+	{
+		Collector.AddReferencedObject(Preview);
+	}
+
+	FString FWidgetPreviewToolkit::GetReferencerName() const
+	{
+		return TEXT("FWidgetPreviewToolkit");
+	}
+
+	UWidgetPreview* FWidgetPreviewToolkit::GetPreview() const
+	{
+		return Preview;
+	}
+
+	FWidgetPreviewToolkitStateBase* FWidgetPreviewToolkit::GetState() const
+	{
+		return CurrentState;
+	}
+
+	IWidgetPreviewToolkit::FOnStateChanged& FWidgetPreviewToolkit::OnStateChanged()
+	{
+		return OnStateChangedDelegate;
+	}
+
+	TSharedPtr<FWidgetBlueprintEditor> FWidgetPreviewToolkit::GetEditor() const
+	{
+		if (TSharedPtr<FWidgetBlueprintEditor> Editor = WeakEditor.Pin())
+		{
+			return Editor;
+		}
+
+		return nullptr;
+	}
+
+	TSharedPtr<FWidgetPreviewScene> FWidgetPreviewToolkit::GetPreviewScene()
+	{
+		if (!PreviewScene.IsValid())
+		{
+			const TSharedRef<FWidgetPreviewToolkit>& Self = SharedThis(this);
+			PreviewScene = MakeShared<FWidgetPreviewScene>(Self);
+		}
+
+		return PreviewScene;
+	}
+
+	UWorld* FWidgetPreviewToolkit::GetPreviewWorld()
+	{
+		if (TSharedPtr<FWidgetPreviewScene> PreviewScenePtr = GetPreviewScene())
+		{
+			return PreviewScenePtr->GetWorld();
+		}
+
+		return nullptr;
+	}
+
+	bool FWidgetPreviewToolkit::ShouldUpdate() const
+	{
+		if (CurrentState)
+		{
+			return CurrentState->CanTick();
+		}
+
+		return bIsFocused;
+	}
+
+	void FWidgetPreviewToolkit::OnWidgetChanged(const EWidgetPreviewWidgetChangeType InChangeType)
+	{
+		ResolveState();
+	}
+
+	void FWidgetPreviewToolkit::OnFocusChanging(
+		const FFocusEvent& InFocusEvent,
+		const FWeakWidgetPath& InOldWidgetPath, const TSharedPtr<SWidget>& InOldWidget,
+		const FWidgetPath& InNewWidgetPath, const TSharedPtr<SWidget>& InNewWidget)
+	{
+		if (IsHosted())
+		{
+			const SWidget* ToolkitParentWidget = GetToolkitHost()->GetParentWidget().ToSharedPtr().Get();
+			const bool bToolkitInNewWidgetPath = InNewWidgetPath.ContainsWidget(ToolkitParentWidget);
+			if (bIsFocused && !bToolkitInNewWidgetPath)
+			{
+				// Focus lost
+				bIsFocused = false;
+				ResolveState();
+			}
+			else if (!bIsFocused && bToolkitInNewWidgetPath)
+			{
+				// Focus received
+				bIsFocused = true;
+				ResolveState();
+			}
+		}
+	}
+
+	void FWidgetPreviewToolkit::SetState(FWidgetPreviewToolkitStateBase* InNewState)
+	{
+		FWidgetPreviewToolkitStateBase* OldState = CurrentState;
+		FWidgetPreviewToolkitStateBase* NewState = InNewState;
+
+		if (OldState != NewState)
+		{
+			if (OldState)
+			{
+				OldState->OnExit(NewState);
+			}
+
+			if (NewState)
+			{
+				NewState->OnEnter(OldState);
+			}
+
+			CurrentState = NewState;
+			OnStateChanged().Broadcast(OldState, NewState);
+		}
+	}
+
+	void FWidgetPreviewToolkit::ResolveState()
+	{
+		FWidgetPreviewToolkitStateBase* NewState = nullptr;
+
+		if (!bIsFocused)
+		{
+			NewState = &BackgroundState;
+		}
+		else
+		{
+			TArray<const UUserWidget*> FailedWidgets;
+			if (!GetPreview()->CanCallInitializedWithoutPlayerContext(true, FailedWidgets))
+			{
+				UnsupportedWidgetState.SetUnsupportedWidgets(FailedWidgets);
+				NewState = &UnsupportedWidgetState;
+			}
+
+			// If we're here, the current state should be valid/running
+			if (NewState == nullptr)
+			{
+				NewState = &RunningState;
+			}
+		}
+
+		SetState(NewState);
+	}
+
+	void FWidgetPreviewToolkit::ResetPreview()
+	{
+		if (Preview)
+		{
+			// Don't need the returned instance, just need to have it rebuild
+			Preview->GetOrCreateWidgetInstance(GetPreviewWorld(), true);
+		}
+	}
+
+	TSharedRef<SDockTab> FWidgetPreviewToolkit::SpawnTab_Viewport(const FSpawnTabArgs& Args)
+	{
+		check(Args.GetTabId() == ViewportTabID);
+
+		FAssetEditorViewportConstructionArgs ViewportArgs;
+		ViewportArgs.ViewportType = LVT_Perspective;
+		ViewportArgs.bRealtime = true;
+
+		const TSharedRef<FWidgetPreviewToolkit>& Self = SharedThis(this);
+		TSharedRef<SEditorViewport> ViewportWidget = SNew(SWidgetPreviewViewport, Self);
+
+		TSharedRef<SDockTab> DockTab = SNew(SDockTab);
+
+		DockTab
+		->SetContent(
+			SNew(SOverlay)
+			+ SOverlay::Slot()
+			.HAlign(HAlign_Fill)
+			.VAlign(VAlign_Fill)
+			[
+				ViewportWidget
+			]
+
+			+ SOverlay::Slot()
+			.HAlign(HAlign_Fill)
+			.VAlign(VAlign_Fill)
+			[
+				SNew(SWidgetPreview, Self)
+				.IsEnabled(this, &FWidgetPreviewToolkit::ShouldUpdate)
+			]
+
+			+ SOverlay::Slot()
+			.HAlign(HAlign_Center)
+			.VAlign(VAlign_Center)
+			[
+				SNew(SWidgetPreviewStatus, Self)
+			]
+		);
+
+		return DockTab;
+	}
+
+	TSharedRef<SDockTab> FWidgetPreviewToolkit::SpawnTab_Details(const FSpawnTabArgs& Args)
+	{
+		check(Args.GetTabId() == DetailsTabID);
+		return SNew(SDockTab)
+		[
+			SNew(SWidgetPreviewDetails, Preview)
+		];
+	}
+
+	TSharedRef<SDockTab> FWidgetPreviewToolkit::SpawnTab_PreviewSceneSettings(const FSpawnTabArgs& Args)
+	{
+		check(Args.GetTabId().TabType == PreviewSceneSettingsTabID);
+		return SNew(SDockTab)
+		.Label( LOCTEXT("StaticMeshPreviewScene_TabTitle", "Preview Scene Settings") )
+		[
+			PreviewSettingsWidget.IsValid() ? PreviewSettingsWidget.ToSharedRef() : SNullWidget::NullWidget
+		];
+	}
+
+	TSharedRef<SDockTab> FWidgetPreviewToolkit::SpawnTab_MessageLog(const FSpawnTabArgs& Args)
+	{
+		check(Args.GetTabId().TabType == MessageLogTabID);
+		return SNew(SDockTab)
+		[
+			MessageLogWidget.ToSharedRef()
+		];
+	}
+}
+
+#undef LOCTEXT_NAMESPACE
diff --git a/Engine/Plugins/Editor/UMGWidgetPreview/Source/UMGWidgetPreview/Private/WidgetPreviewToolkit.h b/Engine/Plugins/Editor/UMGWidgetPreview/Source/UMGWidgetPreview/Private/WidgetPreviewToolkit.h
new file mode 100644
--- /dev/null
+++ b/Engine/Plugins/Editor/UMGWidgetPreview/Source/UMGWidgetPreview/Private/WidgetPreviewToolkit.h
@@ -0,0 +1,218 @@
+// Copyright Epic Games, Inc. All Rights Reserved.
+
+#pragma once
+
+#include "AdvancedPreviewScene.h"
+#include "AdvancedPreviewSceneModule.h"
+#include "Framework/Docking/TabManager.h"
+#include "Logging/TokenizedMessage.h"
+#include "Misc/DataValidation/Fixer.h"
+#include "Templates/SharedPointer.h"
+#include "Toolkits/BaseToolkit.h"
+#include "Tools/BaseAssetToolkit.h"
+#include "WidgetPreview.h"
+#include "WorkflowOrientedApp/WorkflowTabManager.h"
+
+class UWidgetPreviewEditor;
+class FWidgetBlueprintEditor;
+class IMessageLogListing;
+class UWidgetPreview;
+
+namespace UE::UMGWidgetPreview::Private
+{
+	struct FWidgetPreviewabilityFixer : DataValidation::IFixer
+	{
+		virtual EFixApplicability GetApplicability(int32 FixIndex) const override;
+		virtual FFixResult ApplyFix(int32 FixIndex) override;
+
+		TWeakObjectPtr<const UUserWidget> WeakUserWidget;
+
+		static TSharedRef<FWidgetPreviewabilityFixer> Create(const UUserWidget* InUserWidget);
+	};
+
+	class FWidgetPreviewToolkit;
+
+	/** Encapsulates the state needed to run the preview world. */
+	class FWidgetPreviewScene
+	: public FTickableEditorObject
+	{
+	public:
+		explicit FWidgetPreviewScene(const TSharedRef<FWidgetPreviewToolkit>& InPreviewToolkit);
+		virtual ~FWidgetPreviewScene() override = default;
+
+		//~ Begin FTickableEditorObject
+		virtual void Tick(float DeltaTime) override;
+		virtual ETickableTickType GetTickableTickType() const override;
+		virtual TStatId GetStatId() const override;
+		//~ End FTickableEditorObject
+
+	public:
+		UWorld* GetWorld() const;
+
+		TSharedRef<FAdvancedPreviewScene> GetPreviewScene() const;
+
+	private:
+		TWeakPtr<FWidgetPreviewToolkit> WeakToolkit;
+		TSharedPtr<FAdvancedPreviewScene> PreviewScene;
+	};
+
+	struct FWidgetPreviewToolkitStateBase
+	{
+		explicit FWidgetPreviewToolkitStateBase(const FName& Id)
+			: Id(Id)
+		{
+		}
+
+		virtual ~FWidgetPreviewToolkitStateBase() = default;
+
+		FName GetId() const;
+		const TSharedPtr<FTokenizedMessage>& GetStatusMessage() const;
+		bool CanTick() const;
+		bool ShouldOverlayStatusMessage() const;
+
+		virtual void OnEnter(const FWidgetPreviewToolkitStateBase* InFromState);
+		virtual void OnExit(const FWidgetPreviewToolkitStateBase* InToState);
+
+	protected:
+		FName Id;
+		TSharedPtr<FTokenizedMessage> StatusMessage;
+		bool bCanTick = false;
+		bool bShouldOverlayMessage = false;
+	};
+
+	struct FWidgetPreviewToolkitPausedState : FWidgetPreviewToolkitStateBase
+	{
+		FWidgetPreviewToolkitPausedState();
+	};
+
+	struct FWidgetPreviewToolkitBackgroundState : FWidgetPreviewToolkitPausedState
+	{
+		FWidgetPreviewToolkitBackgroundState();
+	};
+
+	struct FWidgetPreviewToolkitUnsupportedWidgetState : FWidgetPreviewToolkitPausedState
+	{
+		FWidgetPreviewToolkitUnsupportedWidgetState();
+
+		void SetUnsupportedWidgets(const TArray<const UUserWidget*>& InWidgets);
+
+	private:
+		void ResetStatusMessage();
+
+	private:
+		TArray<TWeakObjectPtr<const UUserWidget>> UnsupportedWidgets;
+	};
+
+	struct FWidgetPreviewToolkitRunningState : FWidgetPreviewToolkitStateBase
+	{
+		FWidgetPreviewToolkitRunningState();
+	};
+
+	class IWidgetPreviewToolkit
+	{
+	public:
+		virtual ~IWidgetPreviewToolkit() = default;
+
+		using FOnStateChanged = TMulticastDelegate<void(
+			FWidgetPreviewToolkitStateBase* InOldState, FWidgetPreviewToolkitStateBase* InNewState)>;
+
+		virtual UWidgetPreview* GetPreview() const = 0;
+		virtual FWidgetPreviewToolkitStateBase* GetState() const = 0;
+		virtual FOnStateChanged& OnStateChanged() = 0;
+	};
+
+	class FWidgetPreviewToolkit
+	: public FBaseAssetToolkit
+	  , public FGCObject
+	  , public IWidgetPreviewToolkit
+	{
+	public:
+		explicit FWidgetPreviewToolkit(UWidgetPreviewEditor* InOwningEditor);
+		virtual ~FWidgetPreviewToolkit() override;
+
+		//~ Begin FBaseAssetToolkit
+		virtual void CreateWidgets() override;
+		virtual void RegisterToolbar() override;
+		//~ End FBaseAssetToolkit
+
+		//~ Begin FAssetEditorToolkit
+		virtual void RegisterTabSpawners(const TSharedRef<FTabManager>& TabManager) override;
+		virtual void UnregisterTabSpawners(const TSharedRef<FTabManager>& InTabManager) override;
+		virtual void PostInitAssetEditor() override;
+		//~ End FAssetEditorToolkit
+
+		//~ Begin IToolkit
+		virtual FText GetToolkitName() const override;
+		virtual FName GetToolkitFName() const override;
+		virtual FText GetBaseToolkitName() const override;
+		virtual FString GetWorldCentricTabPrefix() const override;
+		virtual FLinearColor GetWorldCentricTabColorScale() const override;
+		//~ End IToolkit
+
+		//~ Begin FGCObject
+		virtual void AddReferencedObjects(FReferenceCollector& Collector) override;
+		virtual FString GetReferencerName() const override;
+		//~ End FGCObject
+
+		//~ Begin IWidgetPreviewToolkit
+		virtual UWidgetPreview* GetPreview() const override;
+		virtual FWidgetPreviewToolkitStateBase* GetState() const override;
+		virtual FOnStateChanged& OnStateChanged() override;
+		//~ End IWidgetPreviewToolkit
+
+		TSharedPtr<FWidgetBlueprintEditor> GetEditor() const;
+		TSharedPtr<FWidgetPreviewScene> GetPreviewScene();
+		UWorld* GetPreviewWorld();
+
+	public:
+		static const FLazyName PreviewSceneSettingsTabID;
+		static const FLazyName MessageLogTabID;
+
+	protected:
+		bool ShouldUpdate() const;
+
+		void OnWidgetChanged(const EWidgetPreviewWidgetChangeType InChangeType);
+
+		void OnFocusChanging(
+			const FFocusEvent& InFocusEvent,
+			const FWeakWidgetPath& InOldWidgetPath, const TSharedPtr<SWidget>& InOldWidget,
+			const FWidgetPath& InNewWidgetPath, const TSharedPtr<SWidget>& InNewWidget);
+
+		/** If the given state is different to the current state, this will handle transitions and events. */
+		void SetState(FWidgetPreviewToolkitStateBase* InNewState);
+
+		/** Resolve and set the current state based on various conditions. */
+		void ResolveState();
+
+		void ResetPreview();
+
+		virtual TSharedRef<SDockTab> SpawnTab_Viewport(const FSpawnTabArgs& Args) override;
+		virtual TSharedRef<SDockTab> SpawnTab_Details(const FSpawnTabArgs& Args) override;
+		TSharedRef<SDockTab> SpawnTab_PreviewSceneSettings(const FSpawnTabArgs& Args);
+		TSharedRef<SDockTab> SpawnTab_MessageLog(const FSpawnTabArgs& Args);
+
+	private:
+		TWeakPtr<FWidgetBlueprintEditor> WeakEditor;
+		TObjectPtr<UWidgetPreview> Preview;
+
+		TSharedPtr<FWidgetPreviewScene> PreviewScene;
+		FAdvancedPreviewSceneModule::FOnPreviewSceneChanged OnPreviewSceneChangedDelegate;
+		TSharedPtr<SWidget> PreviewSettingsWidget;
+
+		TSharedPtr<IMessageLogListing> MessageLogListing;
+		TSharedPtr<SWidget> MessageLogWidget;
+
+		bool bIsFocused = false;
+		FWorkflowAllowedTabSet TabFactories;
+
+		FDelegateHandle OnWidgetChangedHandle;
+		FDelegateHandle OnFocusChangingHandle;
+		FOnStateChanged OnStateChangedDelegate;
+
+		FWidgetPreviewToolkitStateBase* CurrentState = nullptr;
+		FWidgetPreviewToolkitPausedState PausedState;
+		FWidgetPreviewToolkitBackgroundState BackgroundState;
+		FWidgetPreviewToolkitUnsupportedWidgetState UnsupportedWidgetState;
+		FWidgetPreviewToolkitRunningState RunningState;
+	};
+}
diff --git a/Engine/Plugins/Editor/UMGWidgetPreview/Source/UMGWidgetPreview/Private/WidgetPreviewTypesPrivate.h b/Engine/Plugins/Editor/UMGWidgetPreview/Source/UMGWidgetPreview/Private/WidgetPreviewTypesPrivate.h
new file mode 100644
--- /dev/null
+++ b/Engine/Plugins/Editor/UMGWidgetPreview/Source/UMGWidgetPreview/Private/WidgetPreviewTypesPrivate.h
@@ -0,0 +1,26 @@
+// Copyright Epic Games, Inc. All Rights Reserved.
+
+#pragma once
+
+#include "Blueprint/UserWidget.h"
+#include "Blueprint/WidgetBlueprintGeneratedClass.h"
+#include "WidgetBlueprint.h"
+
+namespace UE::UMGWidgetPreview::Private
+{
+	struct FWidgetTypeTuple
+	{
+		explicit FWidgetTypeTuple(const UUserWidget* InUserWidget)
+		{
+			check(InUserWidget);
+
+			ClassDefaultObject = InUserWidget;
+			BlueprintGeneratedClass = Cast<UWidgetBlueprintGeneratedClass>(ClassDefaultObject->GetClass());
+			Blueprint = Cast<UWidgetBlueprint>(BlueprintGeneratedClass->ClassGeneratedBy);
+		}
+
+		const UUserWidget* ClassDefaultObject = nullptr;
+		UWidgetBlueprint* Blueprint = nullptr;
+		UWidgetBlueprintGeneratedClass* BlueprintGeneratedClass = nullptr;
+	};
+}
diff --git a/Engine/Plugins/Editor/UMGWidgetPreview/Source/UMGWidgetPreview/Private/Widgets/SWidgetPreview.cpp b/Engine/Plugins/Editor/UMGWidgetPreview/Source/UMGWidgetPreview/Private/Widgets/SWidgetPreview.cpp
new file mode 100644
--- /dev/null
+++ b/Engine/Plugins/Editor/UMGWidgetPreview/Source/UMGWidgetPreview/Private/Widgets/SWidgetPreview.cpp
@@ -0,0 +1,90 @@
+// Copyright Epic Games, Inc. All Rights Reserved.
+
+#include "SWidgetPreview.h"
+
+#include "Blueprint/UserWidget.h"
+#include "Editor.h"
+#include "Engine/Engine.h"
+#include "Engine/World.h"
+#include "WidgetPreview.h"
+#include "Slate/SRetainerWidget.h"
+#include "WidgetPreviewToolkit.h"
+#include "Widgets/Layout/SBorder.h"
+
+#define LOCTEXT_NAMESPACE "SWidgetPreview"
+
+namespace UE::UMGWidgetPreview::Private
+{
+	void SWidgetPreview::Construct(const FArguments& Args, const TSharedRef<FWidgetPreviewToolkit>& InToolkit)
+	{
+		WeakToolkit = InToolkit;
+
+		OnWidgetChangedHandle = InToolkit->GetPreview()->OnWidgetChanged().AddSP(this, &SWidgetPreview::OnWidgetChanged);
+
+		CreatedSlateWidget = SNullWidget::NullWidget;
+
+		ContainerWidget = SNew(SBorder)
+		[
+			CreatedSlateWidget.ToSharedRef()
+		];
+
+		OnWidgetChanged(EWidgetPreviewWidgetChangeType::Assignment);
+
+		ChildSlot
+		[
+			SNew(SOverlay)
+			+ SOverlay::Slot()
+			[
+				SAssignNew(RetainerWidget, SRetainerWidget)
+				[
+					ContainerWidget.ToSharedRef()
+				]
+			]
+		];
+	}
+
+	SWidgetPreview::~SWidgetPreview()
+    {
+    	if (const TSharedPtr<FWidgetPreviewToolkit> Toolkit = WeakToolkit.Pin())
+    	{
+    		if (UWidgetPreview* Preview = Toolkit->GetPreview())
+    		{
+    			Preview->OnWidgetChanged().Remove(OnWidgetChangedHandle);
+    		}
+    	}
+    }
+
+	void SWidgetPreview::OnWidgetChanged(const EWidgetPreviewWidgetChangeType InChangeType)
+	{
+		if (const TSharedPtr<FWidgetPreviewToolkit> Toolkit = WeakToolkit.Pin())
+		{
+			if (UWidgetPreview* Preview = Toolkit->GetPreview())
+			{
+				UWorld* World = GetWorld();
+
+				if (UUserWidget* PreviewWidget = Preview->GetOrCreateWidgetInstance(World))
+				{
+					CreatedSlateWidget = PreviewWidget->TakeWidget();
+				}
+				else
+				{
+					CreatedSlateWidget = SNullWidget::NullWidget;
+				}
+
+				ContainerWidget->SetContent(CreatedSlateWidget.ToSharedRef());
+			}
+		}
+	}
+
+	UWorld* SWidgetPreview::GetWorld() const
+	{
+		if (const TSharedPtr<FWidgetPreviewToolkit> Toolkit = WeakToolkit.Pin())
+		{
+			return Toolkit->GetPreviewWorld();
+		}
+
+		return nullptr;
+	}
+}
+
+#undef LOCTEXT_NAMESPACE
diff --git a/Engine/Plugins/Editor/UMGWidgetPreview/Source/UMGWidgetPreview/Private/Widgets/SWidgetPreview.h b/Engine/Plugins/Editor/UMGWidgetPreview/Source/UMGWidgetPreview/Private/Widgets/SWidgetPreview.h
new file mode 100644
--- /dev/null
+++ b/Engine/Plugins/Editor/UMGWidgetPreview/Source/UMGWidgetPreview/Private/Widgets/SWidgetPreview.h
@@ -0,0 +1,47 @@
+// Copyright Epic Games, Inc. All Rights Reserved.
+
+#pragma once
+
+#include "UObject/StrongObjectPtr.h"
+#include "UObject/ObjectPtr.h"
+#include "WidgetPreview.h"
+#include "Widgets/SCompoundWidget.h"
+
+class UWidget;
+class UUserWidget;
+class SBorder;
+class SRetainerWidget;
+
+namespace UE::UMGWidgetPreview::Private
+{
+	class FWidgetPreviewToolkit;
+
+	class SWidgetPreview
+		: public SCompoundWidget
+	{
+	public:
+		using FSlotWidgetMap = TMap<FName, TObjectPtr<UWidget>>;
+
+		SLATE_BEGIN_ARGS(SWidgetPreview) {}
+		SLATE_END_ARGS()
+
+		virtual ~SWidgetPreview() override;
+
+		void Construct(const FArguments& Args, const TSharedRef<FWidgetPreviewToolkit>& InToolkit);
+
+	private:
+		void OnWidgetChanged(const EWidgetPreviewWidgetChangeType InChangeType);
+
+		/** Convenience method to get world from the associated viewport. */
+		UWorld* GetWorld() const;
+
+	private:
+		TWeakPtr<FWidgetPreviewToolkit> WeakToolkit;
+
+		TSharedPtr<SRetainerWidget> RetainerWidget;
+		TSharedPtr<SBorder> ContainerWidget;
+		TSharedPtr<SWidget> CreatedSlateWidget;
+
+		FDelegateHandle OnWidgetChangedHandle;
+	};
+}
diff --git a/Engine/Plugins/Editor/UMGWidgetPreview/Source/UMGWidgetPreview/Private/Widgets/SWidgetPreviewDetails.cpp b/Engine/Plugins/Editor/UMGWidgetPreview/Source/UMGWidgetPreview/Private/Widgets/SWidgetPreviewDetails.cpp
new file mode 100644
--- /dev/null
+++ b/Engine/Plugins/Editor/UMGWidgetPreview/Source/UMGWidgetPreview/Private/Widgets/SWidgetPreviewDetails.cpp
@@ -0,0 +1,44 @@
+// Copyright Epic Games, Inc. All Rights Reserved.
+
+#include "SWidgetPreviewDetails.h"
+
+#include "IDetailsView.h"
+#include "Modules/ModuleManager.h"
+#include "PropertyEditorModule.h"
+#include "WidgetPreview.h"
+
+namespace UE::UMGWidgetPreview::Private
+{
+	void SWidgetPreviewDetails::Construct(const FArguments& Args, UWidgetPreview* InPreview)
+	{
+		WeakPreview = InPreview;
+
+		FPropertyEditorModule& PropertyEditor = FModuleManager::LoadModuleChecked<FPropertyEditorModule>("PropertyEditor");
+
+		FDetailsViewArgs DetailsViewArgs;
+		DetailsViewArgs.bShowOptions = false;
+		DetailsViewArgs.bAllowSearch = true;
+		DetailsViewArgs.bShowPropertyMatrixButton = false;
+		DetailsViewArgs.bUpdatesFromSelection = false;
+		DetailsViewArgs.bLockable = false;
+		DetailsViewArgs.NameAreaSettings = FDetailsViewArgs::HideNameArea;
+		DetailsViewArgs.NotifyHook = this;
+
+		DetailsView = PropertyEditor.CreateDetailView(DetailsViewArgs);
+
+		HandleSelectedObjectChanged();
+
+		ChildSlot
+		[
+			DetailsView.ToSharedRef()
+		];
+	}
+
+	void SWidgetPreviewDetails::HandleSelectedObjectChanged()
+	{
+		if (UWidgetPreview* Preview = WeakPreview.Get())
+		{
+			DetailsView->SetObject(Preview);
+		}
+	}
+}
diff --git a/Engine/Plugins/Editor/UMGWidgetPreview/Source/UMGWidgetPreview/Private/Widgets/SWidgetPreviewDetails.h b/Engine/Plugins/Editor/UMGWidgetPreview/Source/UMGWidgetPreview/Private/Widgets/SWidgetPreviewDetails.h
new file mode 100644
--- /dev/null
+++ b/Engine/Plugins/Editor/UMGWidgetPreview/Source/UMGWidgetPreview/Private/Widgets/SWidgetPreviewDetails.h
@@ -0,0 +1,29 @@
+// Copyright Epic Games, Inc. All Rights Reserved.
+
+#pragma once
+
+#include "Misc/NotifyHook.h"
+#include "Widgets/SCompoundWidget.h"
+
+class FWidgetBlueprintEditor;
+class FWidgetPreviewToolkit;
+class IDetailsView;
+class UWidgetPreview;
+
+namespace UE::UMGWidgetPreview::Private
+{
+	class SWidgetPreviewDetails : public SCompoundWidget, public FNotifyHook
+	{
+		SLATE_BEGIN_ARGS(SWidgetPreviewDetails) {}
+		SLATE_END_ARGS()
+
+		void Construct(const FArguments& Args, UWidgetPreview* InPreview);
+
+	private:
+		void HandleSelectedObjectChanged();
+
+	private:
+		TWeakObjectPtr<UWidgetPreview> WeakPreview;
+		TSharedPtr<IDetailsView> DetailsView;
+	};
+}
diff --git a/Engine/Plugins/Editor/UMGWidgetPreview/Source/UMGWidgetPreview/Private/Widgets/SWidgetPreviewStatus.cpp b/Engine/Plugins/Editor/UMGWidgetPreview/Source/UMGWidgetPreview/Private/Widgets/SWidgetPreviewStatus.cpp
new file mode 100644
--- /dev/null
+++ b/Engine/Plugins/Editor/UMGWidgetPreview/Source/UMGWidgetPreview/Private/Widgets/SWidgetPreviewStatus.cpp
@@ -0,0 +1,199 @@
+// Copyright Epic Games, Inc. All Rights Reserved.
+
+#include "SWidgetPreviewStatus.h"
+
+#include "Brushes/SlateRoundedBoxBrush.h"
+#include "WidgetPreviewToolkit.h"
+#include "Widgets/Images/SImage.h"
+#include "Widgets/Input/SButton.h"
+
+#define LOCTEXT_NAMESPACE "WidgetPreviewStatus"
+
+namespace UE::UMGWidgetPreview::Private
+{
+	void SWidgetPreviewStatus::Construct(const FArguments& Args, const TSharedRef<FWidgetPreviewToolkit>& InToolkit)
+	{
+		WeakToolkit = InToolkit;
+
+		OnStateChangedHandle = InToolkit->OnStateChanged().AddSP(this, &SWidgetPreviewStatus::OnStateChanged);
+
+		MessageContainerWidget = SNew(SBox);
+		MessageContainerWidget->SetContent(MakeMessageWidget());
+
+		ChildSlot
+		[
+			SNew(SBorder)
+			.BorderImage(new FSlateRoundedBoxBrush(FStyleColors::Panel, 4.0f, FStyleColors::Hover, 1.0f))
+			.Visibility(this, &SWidgetPreviewStatus::GetStatusVisibility)
+			[
+				SNew(SHorizontalBox)
+				+ SHorizontalBox::Slot()
+				.Padding(20.0f, 20.0f, 12.0f, 20.0f)
+				.VAlign(VAlign_Top)
+				.HAlign(HAlign_Left)
+				.AutoWidth()
+				[
+					SNew(SImage)
+					.Image(this, &SWidgetPreviewStatus::GetSeverityIconBrush)
+				]
+
+				+ SHorizontalBox::Slot()
+				.Padding(0.0f, 20.0f, 8.0f, 20.0f)
+				.VAlign(VAlign_Top)
+				.HAlign(HAlign_Center)
+				.AutoWidth()
+				[
+					MessageContainerWidget.ToSharedRef()
+				]
+			]
+		];
+	}
+
+	SWidgetPreviewStatus::~SWidgetPreviewStatus()
+	{
+		if (const TSharedPtr<FWidgetPreviewToolkit> Toolkit = WeakToolkit.Pin())
+		{
+			Toolkit->OnStateChanged().Remove(OnStateChangedHandle);
+		}
+	}
+
+	void SWidgetPreviewStatus::OnStateChanged(FWidgetPreviewToolkitStateBase* InOldState, FWidgetPreviewToolkitStateBase* InNewState)
+	{
+		MessageContainerWidget->SetContent(MakeMessageWidget());
+	}
+
+	TSharedRef<SWidget> SWidgetPreviewStatus::MakeMessageWidget()
+	{
+		if (const TSharedPtr<FTokenizedMessage>& StatusMessage = GetStatusMessage())
+		{
+			TSharedRef<SVerticalBox> Container = SNew(SVerticalBox);
+
+			TSharedPtr<SHorizontalBox> RowContainer = nullptr;
+			float RowPadding = 0.0f; // Initially 0.0, 8.0 for subsequent rows
+
+			const TArray<TSharedRef<IMessageToken>>& MessageTokens = StatusMessage->GetMessageTokens();
+			for (TSharedRef<IMessageToken> MessageToken : MessageTokens)
+			{
+				const EMessageToken::Type TokenType = MessageToken->GetType();
+				switch (TokenType)
+				{
+				case EMessageToken::Severity:
+					// Already handled, skip
+					break;
+
+				case EMessageToken::Fix:
+					{
+						TSharedRef<FFixToken> FixToken = StaticCastSharedRef<FFixToken>(MessageToken);
+
+						RowContainer->AddSlot()
+						.HAlign(HAlign_Right)
+						[
+							SNew(SButton)
+							.Text(LOCTEXT("WidgetPreviewStatus_FixToken_Label", "Apply Fix"))
+							.ToolTipText(LOCTEXT("WidgetPreviewStatus_FixToken_Tooltip", "Apply the suggested fix."))
+							.OnClicked(FOnClicked::CreateSPLambda(
+								AsShared(),
+								[this, FixToken]
+								{
+									FFixResult FixResult = FixToken->GetFixer()->ApplyFix(FixToken->GetFixIndex());
+									return FReply::Handled();
+								}))
+						];
+					}
+					break;
+
+				default:
+				case EMessageToken::Text:
+					{
+						// All other tokens, including text, always denote the start of a new row
+						RowContainer = SNew(SHorizontalBox);
+						Container->AddSlot()
+						.AutoHeight()
+						.Padding(0.0f, RowPadding, 0.0f, 0.0f)
+						[
+							RowContainer.ToSharedRef()
+						];
+						RowPadding = 8.0f;
+
+						RowContainer->AddSlot()
+						.AutoWidth()
+						[
+							SNew(STextBlock)
+							.Text(MessageToken->ToText())
+						];
+					}
+					break;
+				}
+			}
+
+			return Container;
+		}
+
+		return SNullWidget::NullWidget;
+	}
+
+	TSharedPtr<FTokenizedMessage> SWidgetPreviewStatus::GetStatusMessage() const
+	{
+		if (const TSharedPtr<FWidgetPreviewToolkit> Toolkit = WeakToolkit.Pin())
+		{
+			if (FWidgetPreviewToolkitStateBase* CurrentState = Toolkit->GetState())
+			{
+				return CurrentState->GetStatusMessage();
+			}
+		}
+
+		return nullptr;
+	}
+
+	EVisibility SWidgetPreviewStatus::GetStatusVisibility() const
+	{
+		if (const TSharedPtr<FWidgetPreviewToolkit> Toolkit = WeakToolkit.Pin())
+		{
+			if (FWidgetPreviewToolkitStateBase* CurrentState = Toolkit->GetState())
+			{
+				return CurrentState->ShouldOverlayStatusMessage() ? EVisibility::Visible : EVisibility::Collapsed;
+			}
+		}
+
+		return EVisibility::Collapsed;
+	}
+
+	const FSlateBrush* SWidgetPreviewStatus::GetSeverityIconBrush() const
+	{
+		switch (GetSeverity())
+		{
+		case EMessageSeverity::Error:
+			return FAppStyle::Get().GetBrush("Icons.ErrorWithColor");
+
+		case EMessageSeverity::PerformanceWarning:
+		case EMessageSeverity::Warning:
+			return FAppStyle::Get().GetBrush("Icons.WarningWithColor");
+
+		default:
+		case EMessageSeverity::Info:
+			return FAppStyle::Get().GetBrush("Icons.InfoWithColor");
+		}
+	}
+
+	EMessageSeverity::Type SWidgetPreviewStatus::GetSeverity() const
+	{
+		if (const TSharedPtr<FTokenizedMessage>& StatusMessage = GetStatusMessage())
+		{
+			return StatusMessage->GetSeverity();
+		}
+
+		return EMessageSeverity::Info;
+	}
+
+	FText SWidgetPreviewStatus::GetMessage() const
+	{
+		if (const TSharedPtr<FTokenizedMessage>& StatusMessage = GetStatusMessage())
+		{
+			return StatusMessage->ToText();
+		}
+
+		return FText::GetEmpty();
+	}
+}
+
+#undef LOCTEXT_NAMESPACE
diff --git a/Engine/Plugins/Editor/UMGWidgetPreview/Source/UMGWidgetPreview/Private/Widgets/SWidgetPreviewStatus.h b/Engine/Plugins/Editor/UMGWidgetPreview/Source/UMGWidgetPreview/Private/Widgets/SWidgetPreviewStatus.h
new file mode 100644
--- /dev/null
+++ b/Engine/Plugins/Editor/UMGWidgetPreview/Source/UMGWidgetPreview/Private/Widgets/SWidgetPreviewStatus.h
@@ -0,0 +1,45 @@
+// Copyright Epic Games, Inc. All Rights Reserved.
+
+#pragma once
+
+#include "Misc/NotifyHook.h"
+#include "WidgetPreviewToolkit.h"
+#include "Widgets/SCompoundWidget.h"
+#include "Widgets/Layout/SBox.h"
+
+class UWidgetPreview;
+class FWidgetBlueprintEditor;
+class FWidgetPreviewToolkit;
+class IDetailsView;
+
+namespace UE::UMGWidgetPreview::Private
+{
+	class SWidgetPreviewStatus
+		: public SCompoundWidget
+	{
+		SLATE_BEGIN_ARGS(SWidgetPreviewStatus) {}
+		SLATE_END_ARGS()
+
+		void Construct(const FArguments& Args, const TSharedRef<FWidgetPreviewToolkit>& InToolkit);
+
+		virtual ~SWidgetPreviewStatus() override;
+
+	private:
+		void OnStateChanged(FWidgetPreviewToolkitStateBase* InOldState, FWidgetPreviewToolkitStateBase* InNewState);
+
+		TSharedRef<SWidget> MakeMessageWidget();
+
+		TSharedPtr<FTokenizedMessage> GetStatusMessage() const;
+
+		EVisibility GetStatusVisibility() const;
+		const FSlateBrush* GetSeverityIconBrush() const;
+		EMessageSeverity::Type GetSeverity() const;
+		FText GetMessage() const;
+
+	private:
+		TWeakPtr<FWidgetPreviewToolkit> WeakToolkit;
+		TSharedPtr<SBox> MessageContainerWidget;
+
+		FDelegateHandle OnStateChangedHandle;
+	};
+}
diff --git a/Engine/Plugins/Editor/UMGWidgetPreview/Source/UMGWidgetPreview/Private/Widgets/SWidgetPreviewViewport.cpp b/Engine/Plugins/Editor/UMGWidgetPreview/Source/UMGWidgetPreview/Private/Widgets/SWidgetPreviewViewport.cpp
new file mode 100644
--- /dev/null
+++ b/Engine/Plugins/Editor/UMGWidgetPreview/Source/UMGWidgetPreview/Private/Widgets/SWidgetPreviewViewport.cpp
@@ -0,0 +1,44 @@
+// Copyright Epic Games, Inc. All Rights Reserved.
+
+#include "SWidgetPreviewViewport.h"
+
+#include "WidgetPreviewToolkit.h"
+
+namespace UE::UMGWidgetPreview::Private
+{
+	void SWidgetPreviewViewport::Construct(const FArguments& InArgs, const TSharedRef<FWidgetPreviewToolkit>& InToolkit)
+	{
+		WeakToolkit = InToolkit;
+
+		SEditorViewport::Construct(
+			SEditorViewport::FArguments()
+			.ViewportSize(InArgs._ViewportSize));
+	}
+
+	TSharedRef<FEditorViewportClient> SWidgetPreviewViewport::MakeEditorViewportClient()
+	{
+		FPreviewScene* PreviewScene = nullptr;
+		if (TSharedPtr<FWidgetPreviewToolkit> Toolkit = WeakToolkit.Pin())
+		{
+			PreviewScene = &Toolkit->GetPreviewScene()->GetPreviewScene().Get();
+		}
+
+		TSharedRef<FEditorViewportClient> ViewportClient = MakeShared<FEditorViewportClient>(nullptr, PreviewScene);
+		ViewportClient->SetViewLocation(FVector::ZeroVector);
+		ViewportClient->SetViewRotation(FRotator(-25.0f, -135.0f, 0.0f));
+		ViewportClient->SetViewLocationForOrbiting(FVector::ZeroVector, 500);
+		ViewportClient->bSetListenerPosition = false;
+		ViewportClient->EngineShowFlags.EnableAdvancedFeatures();
+		ViewportClient->EngineShowFlags.SetGrid(true);
+		ViewportClient->EngineShowFlags.SetLighting(true);
+		ViewportClient->EngineShowFlags.SetIndirectLightingCache(true);
+		ViewportClient->EngineShowFlags.SetPostProcessing(true);
+
+		// @todo: temp settings, get from elsewhere?
+		ViewportClient->ExposureSettings.bFixed = true;
+
+		ViewportClient->Invalidate();
+
+		return ViewportClient;
+	}
+}
diff --git a/Engine/Plugins/Editor/UMGWidgetPreview/Source/UMGWidgetPreview/Private/Widgets/SWidgetPreviewViewport.h b/Engine/Plugins/Editor/UMGWidgetPreview/Source/UMGWidgetPreview/Private/Widgets/SWidgetPreviewViewport.h
new file mode 100644
--- /dev/null
+++ b/Engine/Plugins/Editor/UMGWidgetPreview/Source/UMGWidgetPreview/Private/Widgets/SWidgetPreviewViewport.h
@@ -0,0 +1,27 @@
+// Copyright Epic Games, Inc. All Rights Reserved.
+
+#pragma once
+
+#include "SAssetEditorViewport.h"
+
+namespace UE::UMGWidgetPreview::Private
+{
+	class FWidgetPreviewToolkit;
+
+	class SWidgetPreviewViewport
+		: public SEditorViewport
+	{
+	public:
+		SLATE_BEGIN_ARGS(SWidgetPreviewViewport) {}
+		SLATE_ATTRIBUTE(FVector2D, ViewportSize);
+		SLATE_END_ARGS()
+
+		void Construct(const FArguments& InArgs, const TSharedRef<FWidgetPreviewToolkit>& InToolkit);
+
+	private:
+		virtual TSharedRef<FEditorViewportClient> MakeEditorViewportClient() override;
+
+	private:
+		TWeakPtr<FWidgetPreviewToolkit> WeakToolkit;
+	};
+}
diff --git a/Engine/Plugins/Editor/UMGWidgetPreview/Source/UMGWidgetPreview/UMGWidgetPreview.Build.cs b/Engine/Plugins/Editor/UMGWidgetPreview/Source/UMGWidgetPreview/UMGWidgetPreview.Build.cs
new file mode 100644
--- /dev/null
+++ b/Engine/Plugins/Editor/UMGWidgetPreview/Source/UMGWidgetPreview/UMGWidgetPreview.Build.cs
@@ -0,0 +1,34 @@
+﻿// Copyright Epic Games, Inc. All Rights Reserved.
+
+using UnrealBuildTool;
+
+public class UMGWidgetPreview : ModuleRules
+{
+	public UMGWidgetPreview(ReadOnlyTargetRules Target) : base(Target)
+	{
+		PCHUsage = ModuleRules.PCHUsageMode.UseExplicitOrSharedPCHs;
+
+		PrivateDependencyModuleNames.AddRange(
+			new string[]
+			{
+				"AdvancedPreviewScene",
+				"AssetDefinition",
+				"ContentBrowser",
+				"Core",
+				"CoreUObject",
+                "DataValidation",
+				"EditorSubsystem",
+				"Engine",
+				"InputCore",
+				"MessageLog",
+				"Projects",
+				"PropertyEditor",
+				"Slate",
+				"SlateCore",
+				"ToolMenus",
+				"UMG",
+				"UMGEditor",
+				"UnrealEd",
+			});
+	}
+}
diff --git a/Engine/Plugins/Editor/UMGWidgetPreview/UMGWidgetPreview.uplugin b/Engine/Plugins/Editor/UMGWidgetPreview/UMGWidgetPreview.uplugin
new file mode 100644
--- /dev/null
+++ b/Engine/Plugins/Editor/UMGWidgetPreview/UMGWidgetPreview.uplugin
@@ -0,0 +1,31 @@
+﻿{
+	"FileVersion": 3,
+	"Version": 1,
+	"VersionName": "1.0",
+	"FriendlyName": "UMG Widget Preview",
+	"Description": "Quickly preview and debug UMG widgets without running PIE.",
+	"Category": "UI",
+	"CreatedBy": "Epic Games, Inc.",
+	"CreatedByURL": "https://epicgames.com",
+	"DocsURL": "",
+	"MarketplaceURL": "",
+	"SupportURL": "",
+	"CanContainContent": true,
+	"IsBetaVersion": false,
+	"IsExperimentalVersion": true,
+	"Installed": false,
+	"EnabledByDefault": true,
+	"Modules": [
+		{
+			"Name": "UMGWidgetPreview",
+			"Type": "Editor",
+			"LoadingPhase": "Default"
+		}
+	],
+	"Plugins": [
+		{
+			"Name": "DataValidation",
+			"Enabled": true
+		}
+	]
+}
