[UMG Widget Preview] Fixed compatibility check recursion
#rnx

diff --git a/Engine/Plugins/Editor/UMGWidgetPreview/Source/UMGWidgetPreview/Private/WidgetPreview.cpp b/Engine/Plugins/Editor/UMGWidgetPreview/Source/UMGWidgetPreview/Private/WidgetPreview.cpp
--- a/Engine/Plugins/Editor/UMGWidgetPreview/Source/UMGWidgetPreview/Private/WidgetPreview.cpp
+++ b/Engine/Plugins/Editor/UMGWidgetPreview/Source/UMGWidgetPreview/Private/WidgetPreview.cpp
@@ -135,20 +135,31 @@ const UUserWidget* UWidgetPreview::GetWidgetForSlot(const FName InSlotName) cons
 
 bool UWidgetPreview::CanCallInitializedWithoutPlayerContext(const bool bInRecursive, TArray<const UUserWidget*>& OutFailedWidgets)
 {
+	bool bResult = true;
+
 	if (!LayoutWidgetType.IsNull())
 	{
 		const UUserWidget* LayoutWidgetCDO = LayoutWidgetType.LoadSynchronous()->GetDefaultObject<UUserWidget>();
-		return CanCallInitializedWithoutPlayerContextOnWidget(LayoutWidgetCDO, bInRecursive, OutFailedWidgets);
+		bResult = bResult && CanCallInitializedWithoutPlayerContextOnWidget(LayoutWidgetCDO, bInRecursive, OutFailedWidgets);
+
+		for (const TPair<FName, TSoftClassPtr<UUserWidget>>& SlotWidget : SlotWidgets)
+		{
+			if (!SlotWidget.Value.IsNull())
+			{
+				const UUserWidget* SlotWidgetCDO = SlotWidget.Value.LoadSynchronous()->GetDefaultObject<UUserWidget>();
+				bResult = bResult && CanCallInitializedWithoutPlayerContextOnWidget(SlotWidgetCDO, bInRecursive, OutFailedWidgets);
+			}
+		}
 	}
 
 	if (!WidgetType.IsNull())
 	{
 		const UUserWidget* WidgetCDO = WidgetType.LoadSynchronous()->GetDefaultObject<UUserWidget>();
-		return CanCallInitializedWithoutPlayerContextOnWidget(WidgetCDO, bInRecursive, OutFailedWidgets);
+		bResult = bResult && CanCallInitializedWithoutPlayerContextOnWidget(WidgetCDO, bInRecursive, OutFailedWidgets);
 	}
 
 	// In case there are no widgets to display, we want to return true
-	return true;
+	return bResult;
 }
 
 bool UWidgetPreview::CanCallInitializedWithoutPlayerContextOnWidget(
@@ -173,21 +184,21 @@ bool UWidgetPreview::CanCallInitializedWithoutPlayerContextOnWidget(
 					OutFailedWidgets.Emplace(WidgetTuple.ClassDefaultObject);
 				}
 			}
-		}
 
-		if (bInRecursive)
-		{
-			if (const INamedSlotInterface* WidgetWithSlots = Cast<INamedSlotInterface>(InWidget))
+			if (bInRecursive)
 			{
-				TArray<FName> SlotNames;
-				WidgetWithSlots->GetSlotNames(SlotNames);
-				if (!SlotNames.IsEmpty())
+				if (const INamedSlotInterface* WidgetWithSlots = Cast<INamedSlotInterface>(AsUserWidget))
 				{
-					for (const FName SlotName : SlotNames)
+					TArray<FName> SlotNames;
+					WidgetWithSlots->GetSlotNames(SlotNames);
+					if (!SlotNames.IsEmpty())
 					{
-						if (const UWidget* SlotWidget = Cast<UWidget>(WidgetWithSlots->GetContentForSlot(SlotName)))
+						for (const FName SlotName : SlotNames)
 						{
-							bResultInternal = bResultInternal && CanCallInitializedWithoutPlayerContextInternal(SlotWidget);
+							if (const UWidget* SlotWidget = Cast<UWidget>(WidgetWithSlots->GetContentForSlot(SlotName)))
+							{
+								bResultInternal = bResultInternal && CanCallInitializedWithoutPlayerContextInternal(SlotWidget);
+							}
 						}
 					}
 				}
