UMG Widget Preview cis fix
#rnx

diff --git a/Engine/Plugins/Editor/UMGWidgetPreview/Source/UMGWidgetPreview/Private/UMGWidgetPreviewModule.cpp b/Engine/Plugins/Editor/UMGWidgetPreview/Source/UMGWidgetPreview/Private/UMGWidgetPreviewModule.cpp
--- a/Engine/Plugins/Editor/UMGWidgetPreview/Source/UMGWidgetPreview/Private/UMGWidgetPreviewModule.cpp
+++ b/Engine/Plugins/Editor/UMGWidgetPreview/Source/UMGWidgetPreview/Private/UMGWidgetPreviewModule.cpp
@@ -5,17 +5,21 @@
 #include "Blueprint/UserWidget.h"
 #include "ContentBrowserMenuContexts.h"
 #include "Editor.h"
+#include "Framework/Commands/UIAction.h"
 #include "MessageLogModule.h"
 #include "Modules/ModuleManager.h"
 #include "PropertyEditorModule.h"
+#include "ToolMenu.h"
+#include "ToolMenuMisc.h"
 #include "ToolMenus.h"
+#include "ToolMenuSection.h"
 #include "UObject/UObjectHash.h"
 #include "WidgetBlueprint.h"
+#include "WidgetPreview.h"
 #include "WidgetPreviewCommands.h"
 #include "WidgetPreviewEditor.h"
 #include "WidgetPreviewLog.h"
 #include "WidgetPreviewStyle.h"
-#include "WidgetPreview.h"
 
 #define LOCTEXT_NAMESPACE "UMGWidgetPreviewModule"
 
diff --git a/Engine/Plugins/Editor/UMGWidgetPreview/Source/UMGWidgetPreview/Private/WidgetPreviewToolkit.cpp b/Engine/Plugins/Editor/UMGWidgetPreview/Source/UMGWidgetPreview/Private/WidgetPreviewToolkit.cpp
--- a/Engine/Plugins/Editor/UMGWidgetPreview/Source/UMGWidgetPreview/Private/WidgetPreviewToolkit.cpp
+++ b/Engine/Plugins/Editor/UMGWidgetPreview/Source/UMGWidgetPreview/Private/WidgetPreviewToolkit.cpp
@@ -2,16 +2,19 @@
 
 #include "WidgetPreviewToolkit.h"
 
-#include "Brushes/SlateRoundedBoxBrush.h"
 #include "DataValidationFixers.h"
 #include "Framework/Application/SlateApplication.h"
-#include "Framework/Application/SlateUser.h"
 #include "Kismet2/KismetEditorUtilities.h"
 #include "MessageLogModule.h"
 #include "ObjectEditorUtils.h"
 #include "SAssetEditorViewport.h"
 #include "ScopedTransaction.h"
 #include "Styling/AppStyle.h"
+#include "ToolMenu.h"
+#include "ToolMenuEntry.h"
+#include "ToolMenuMisc.h"
+#include "ToolMenus.h"
+#include "ToolMenuSection.h"
 #include "WidgetBlueprintEditor.h"
 #include "WidgetPreview.h"
 #include "WidgetPreviewCommands.h"
