[UMG Widget Preview] Fixed log category (trivial fix)

diff --git a/Engine/Plugins/Editor/UMGWidgetPreview/Source/UMGWidgetPreview/Private/UMGWidgetPreviewModule.cpp b/Engine/Plugins/Editor/UMGWidgetPreview/Source/UMGWidgetPreview/Private/UMGWidgetPreviewModule.cpp
--- a/Engine/Plugins/Editor/UMGWidgetPreview/Source/UMGWidgetPreview/Private/UMGWidgetPreviewModule.cpp
+++ b/Engine/Plugins/Editor/UMGWidgetPreview/Source/UMGWidgetPreview/Private/UMGWidgetPreviewModule.cpp
@@ -6,6 +6,7 @@
 #include "ContentBrowserMenuContexts.h"
 #include "Editor.h"
 #include "Framework/Commands/UIAction.h"
+#include "Logging/LogMacros.h"
 #include "MessageLogModule.h"
 #include "Modules/ModuleManager.h"
 #include "PropertyEditorModule.h"
@@ -23,6 +24,8 @@
 
 #define LOCTEXT_NAMESPACE "UMGWidgetPreviewModule"
 
+DEFINE_LOG_CATEGORY(LogWidgetPreview);
+
 namespace UE::UMGWidgetPreview::Private
 {
 	void FUMGWidgetPreviewModule::StartupModule()
