[UMG Widget Preview] Added icon for asset

diff --git a/Engine/Plugins/Editor/UMGWidgetPreview/Source/UMGWidgetPreview/Private/WidgetPreviewStyle.cpp b/Engine/Plugins/Editor/UMGWidgetPreview/Source/UMGWidgetPreview/Private/WidgetPreviewStyle.cpp
--- a/Engine/Plugins/Editor/UMGWidgetPreview/Source/UMGWidgetPreview/Private/WidgetPreviewStyle.cpp
+++ b/Engine/Plugins/Editor/UMGWidgetPreview/Source/UMGWidgetPreview/Private/WidgetPreviewStyle.cpp
@@ -21,12 +21,8 @@ namespace UE::UMGWidgetPreview::Private
 	FWidgetPreviewStyle::FWidgetPreviewStyle()
 		: FSlateStyleSet(StyleName)
 	{
-		const FVector2f Icon16(16.f);
-		const FVector2f Icon20(20.f);
-		const FVector2f Icon22(22.f);
-		const FVector2f Icon25(25.f);
-		const FVector2f Icon32(32.f);
-		const FVector2f Icon40(40.f);
+		const FVector2f Icon16(16.0f);
+		const FVector2f Icon64(64.0f);
 
 		TSharedPtr<IPlugin> Plugin = IPluginManager::Get().FindPlugin(UE_PLUGIN_NAME);
 		check(Plugin.IsValid());
@@ -39,6 +35,7 @@ namespace UE::UMGWidgetPreview::Private
 
 		// Custom Class Icons
 		Set("ClassIcon.WidgetPreview", new IMAGE_BRUSH_SVG(TEXT("Icons/WidgetPreview_16"), Icon16));
+		Set("ClassThumbnail.WidgetPreview", new IMAGE_BRUSH_SVG(TEXT("Icons/WidgetPreview_64"), Icon64));
 
 		FSlateStyleRegistry::RegisterSlateStyle(*this);
 	}
