[UMG Widget Preview] Abstracted Widget type (for customization and to allow future preview asset nesting)

diff --git a/Engine/Plugins/Editor/UMGWidgetPreview/Source/UMGWidgetPreview/Private/WidgetPreview.cpp b/Engine/Plugins/Editor/UMGWidgetPreview/Source/UMGWidgetPreview/Private/WidgetPreview.cpp
--- a/Engine/Plugins/Editor/UMGWidgetPreview/Source/UMGWidgetPreview/Private/WidgetPreview.cpp
+++ b/Engine/Plugins/Editor/UMGWidgetPreview/Source/UMGWidgetPreview/Private/WidgetPreview.cpp
@@ -5,28 +5,100 @@
 #include "Blueprint/UserWidget.h"
 #include "WidgetBlueprint.h"
 #include "WidgetPreviewTypesPrivate.h"
+#include "WidgetPreviewLog.h"
 
-const TArray<FName>& UWidgetPreview::GetLayoutSlotNames()
+FPreviewableWidgetVariant::FPreviewableWidgetVariant(const TSubclassOf<UUserWidget>& InWidgetType)
+	: ObjectPath(InWidgetType)
 {
-	if (!SlotNameCache.IsEmpty())
+	UpdateCachedWidget();
+}
+
+FPreviewableWidgetVariant::FPreviewableWidgetVariant(const UWidgetPreview* InWidgetPreview)
+	: ObjectPath(InWidgetPreview)
+{
+	UpdateCachedWidget();
+}
+
+void FPreviewableWidgetVariant::UpdateCachedWidget()
+{
+	CachedWidgetCDO.Reset();
+	CachedWidgetPreview.Reset();
+
+	if (UObject* ResolvedObject = ObjectPath.TryLoad())
+    {
+    	if (UWidgetPreview* WidgetPreview = Cast<UWidgetPreview>(ResolvedObject))
+    	{
+    		CachedWidgetPreview = WidgetPreview;
+    		CachedWidgetCDO = WidgetPreview->GetWidgetCDO();
+    	}
+		else if (const UWidgetBlueprint* AsBlueprint = Cast<UWidgetBlueprint>(ResolvedObject))
+		{
+			CachedWidgetCDO = Cast<UUserWidget>(AsBlueprint->GeneratedClass->GetDefaultObject<UUserWidget>());
+		}
+		else if (const UClass* AsClass = Cast<UClass>(ResolvedObject))
+		{
+			if (UUserWidget* UserWidget = AsClass->GetDefaultObject<UUserWidget>())
+			{
+				CachedWidgetCDO = UserWidget;
+			}
+		}
+    }
+}
+
+const UUserWidget* FPreviewableWidgetVariant::AsUserWidgetCDO() const
+{
+	if (ObjectPath.IsNull())
+	{
+		return nullptr;
+	}
+
+	if (const UUserWidget* UserWidgetCDO = CachedWidgetCDO.Get())
 	{
-		return SlotNameCache;
+		return UserWidgetCDO;
 	}
 
-	if (!LayoutWidgetType.IsNull())
+	if (UObject* ResolvedObject = ObjectPath.TryLoad())
 	{
-		if (INamedSlotInterface* WidgetWithSlots = Cast<INamedSlotInterface>(LayoutWidgetType.LoadSynchronous()->GetDefaultObject<UUserWidget>()))
+		if (const UClass* AsClass = Cast<UClass>(ResolvedObject))
 		{
-			TArray<FName> SlotNames;
-			WidgetWithSlots->GetSlotNames(SlotNames);
+			if (UUserWidget* UserWidget = AsClass->GetDefaultObject<UUserWidget>())
+			{
+				UE_LOG(LogWidgetPreview, Warning, TEXT("Tried to get the object as a UserWidget (CDO), but it wasn't cached. Ensure you have called Refresh() first."));
+				return UserWidget;
+			}
+		}
+	}
 
-			SlotNameCache = SlotNames;
-			return SlotNameCache;
+	return nullptr;
+}
+
+UWidgetPreview* FPreviewableWidgetVariant::AsWidgetPreview() const
+{
+	if (ObjectPath.IsNull())
+	{
+		return nullptr;
+	}
+
+	if (UWidgetPreview* WidgetPreview = CachedWidgetPreview.Get())
+	{
+		return WidgetPreview;
+	}
+
+	if (UObject* ResolvedObject = ObjectPath.TryLoad())
+	{
+		if (UWidgetPreview* WidgetPreview = Cast<UWidgetPreview>(ResolvedObject))
+		{
+			UE_LOG(LogWidgetPreview, Warning, TEXT("Tried to get the object as a WidgetPreview, but it wasn't cached. Ensure you have called Refresh() first."));
+			return WidgetPreview;
 		}
 	}
 
-	static TArray<FName> Empty;
-	return Empty;
+	return nullptr;
+}
+
+const TArray<FName>& UWidgetPreview::GetWidgetSlotNames() const
+{
+	return SlotNameCache;
 }
 
 UUserWidget* UWidgetPreview::GetOrCreateWidgetInstance(UWorld* InWorld, const bool bInForceRecreate)
@@ -52,7 +124,7 @@ UUserWidget* UWidgetPreview::GetOrCreateWidgetInstance(UWorld* InWorld, const bo
 		return nullptr;
 	}
 
-	if (UUserWidget* Widget = GetWidget())
+	if (const UUserWidget* Widget = GetWidgetCDO())
 	{
 		auto MakeWidget = [InWorld](UClass* InClass) -> UUserWidget*
 		{
@@ -65,17 +137,17 @@ UUserWidget* UWidgetPreview::GetOrCreateWidgetInstance(UWorld* InWorld, const bo
 
 		WidgetInstance = MakeWidget(Widget->GetClass());
 
-		if (!LayoutWidgetType.IsNull() && !SlotWidgets.IsEmpty())
+		if (!WidgetType.ObjectPath.IsNull() && !SlotWidgetTypes.IsEmpty())
 		{
 			TArray<FName> ValidSlotNames;
 			WidgetInstance->GetSlotNames(ValidSlotNames);
 
-			for (const TPair<FName, TSoftClassPtr<UUserWidget>>& SlotWidget : SlotWidgets)
+			for (TPair<FName, FPreviewableWidgetVariant>& SlotWidget : SlotWidgetTypes)
 			{
-				if (!SlotWidget.Value.IsNull()
+				if (!SlotWidget.Value.ObjectPath.IsNull()
 					&& ValidSlotNames.Contains(SlotWidget.Key))
 				{
-					WidgetInstance->SetContentForSlot(SlotWidget.Key, MakeWidget(SlotWidget.Value.LoadSynchronous()));
+					WidgetInstance->SetContentForSlot(SlotWidget.Key, MakeWidget(SlotWidget.Value.AsUserWidgetCDO()->GetClass()));
 				}
 			}
 		}
@@ -93,13 +165,18 @@ UUserWidget* UWidgetPreview::GetOrCreateWidgetInstance(UWorld* InWorld, const bo
 	return nullptr;
 }
 
-UUserWidget* UWidgetPreview::GetWidget() const
+UUserWidget* UWidgetPreview::GetWidgetInstance() const
+{
+	return WidgetInstance;
+}
+
+const UUserWidget* UWidgetPreview::GetWidgetCDO() const
 {
 	// If the LayoutWidget is in use, return it (as the root widget).
-	if (!LayoutWidgetType.IsNull())
+	if (!WidgetType.ObjectPath.IsNull())
 	{
-		UUserWidget* LayoutWidgetCDO = LayoutWidgetType.LoadSynchronous()->GetDefaultObject<UUserWidget>();
-		if (INamedSlotInterface* WidgetWithSlots = Cast<INamedSlotInterface>(LayoutWidgetCDO))
+		const UUserWidget* LayoutWidgetCDO = WidgetType.AsUserWidgetCDO();
+		if (const INamedSlotInterface* WidgetWithSlots = Cast<INamedSlotInterface>(LayoutWidgetCDO))
 		{
 			TArray<FName> SlotNames;
 			WidgetWithSlots->GetSlotNames(SlotNames);
@@ -110,21 +187,21 @@ UUserWidget* UWidgetPreview::GetWidget() const
 		}
 	}
 
-	if (!WidgetType.IsNull())
+	if (!WidgetType.ObjectPath.IsNull())
 	{
-		return WidgetType.LoadSynchronous()->GetDefaultObject<UUserWidget>();
+		return WidgetType.AsUserWidgetCDO();
 	}
 
 	return nullptr;
 }
 
-const UUserWidget* UWidgetPreview::GetWidgetForSlot(const FName InSlotName) const
+const UUserWidget* UWidgetPreview::GetWidgetCDOForSlot(const FName InSlotName) const
 {
-	if (const TSoftClassPtr<UUserWidget>* WidgetInSlot = SlotWidgets.Find(InSlotName))
+	if (const FPreviewableWidgetVariant* WidgetInSlot = SlotWidgetTypes.Find(InSlotName))
 	{
-		if (!WidgetInSlot->IsNull())
+		if (!WidgetInSlot->ObjectPath.IsNull())
 		{
-			return WidgetInSlot->LoadSynchronous()->GetDefaultObject<UUserWidget>();	
+			return WidgetInSlot->AsUserWidgetCDO();
 		}
 
 		UE_LOG(LogTemp, Warning, TEXT("Slot %s has invalid widget."), *InSlotName.ToString());
@@ -133,31 +210,32 @@ const UUserWidget* UWidgetPreview::GetWidgetForSlot(const FName InSlotName) cons
 	return nullptr;
 }
 
+void UWidgetPreview::BeginDestroy()
+{
+	UObject::BeginDestroy();
+
+	CleanupReferences();
+}
+
 bool UWidgetPreview::CanCallInitializedWithoutPlayerContext(const bool bInRecursive, TArray<const UUserWidget*>& OutFailedWidgets)
 {
 	bool bResult = true;
 
-	if (!LayoutWidgetType.IsNull())
+	if (!WidgetType.ObjectPath.IsNull())
 	{
-		const UUserWidget* LayoutWidgetCDO = LayoutWidgetType.LoadSynchronous()->GetDefaultObject<UUserWidget>();
-		bResult = bResult && CanCallInitializedWithoutPlayerContextOnWidget(LayoutWidgetCDO, bInRecursive, OutFailedWidgets);
+		const UUserWidget* WidgetCDO = WidgetType.AsUserWidgetCDO();
+		bResult = bResult && CanCallInitializedWithoutPlayerContextOnWidget(WidgetCDO, bInRecursive, OutFailedWidgets);
 
-		for (const TPair<FName, TSoftClassPtr<UUserWidget>>& SlotWidget : SlotWidgets)
+		for (const TPair<FName, FPreviewableWidgetVariant>& SlotWidget : SlotWidgetTypes)
 		{
-			if (!SlotWidget.Value.IsNull())
+			if (!SlotWidget.Value.ObjectPath.IsNull())
 			{
-				const UUserWidget* SlotWidgetCDO = SlotWidget.Value.LoadSynchronous()->GetDefaultObject<UUserWidget>();
+				const UUserWidget* SlotWidgetCDO = SlotWidget.Value.AsUserWidgetCDO();
 				bResult = bResult && CanCallInitializedWithoutPlayerContextOnWidget(SlotWidgetCDO, bInRecursive, OutFailedWidgets);
 			}
 		}
 	}
 
-	if (!WidgetType.IsNull())
-	{
-		const UUserWidget* WidgetCDO = WidgetType.LoadSynchronous()->GetDefaultObject<UUserWidget>();
-		bResult = bResult && CanCallInitializedWithoutPlayerContextOnWidget(WidgetCDO, bInRecursive, OutFailedWidgets);
-	}
-
 	// In case there are no widgets to display, we want to return true
 	return bResult;
 }
@@ -211,12 +289,12 @@ bool UWidgetPreview::CanCallInitializedWithoutPlayerContextOnWidget(
 	return CanCallInitializedWithoutPlayerContextInternal(InUserWidget);
 }
 
-const TSoftClassPtr<UUserWidget>& UWidgetPreview::GetWidgetType() const
+const FPreviewableWidgetVariant& UWidgetPreview::GetWidgetType() const
 {
 	return WidgetType;
 }
 
-void UWidgetPreview::SetWidgetType(const TSoftClassPtr<UUserWidget>& InWidget)
+void UWidgetPreview::SetWidgetType(const FPreviewableWidgetVariant& InWidget)
 {
 	if (WidgetType != InWidget)
 	{
@@ -228,33 +306,16 @@ void UWidgetPreview::SetWidgetType(const TSoftClassPtr<UUserWidget>& InWidget)
 	}
 }
 
-const TSoftClassPtr<UUserWidget>& UWidgetPreview::GetLayoutWidgetType() const
-{
-	return LayoutWidgetType;
-}
-
-void UWidgetPreview::SetLayoutWidgetType(const TSoftClassPtr<UUserWidget>& InWidget)
-{
-	if (LayoutWidgetType != InWidget)
-	{
-		LayoutWidgetType = InWidget;
-		WidgetInstance = nullptr;
-		UpdateWidgets();
-
-		OnWidgetChanged().Broadcast(EWidgetPreviewWidgetChangeType::Assignment);
-	}
-}
-
-const TMap<FName, TSoftClassPtr<UUserWidget>>& UWidgetPreview::GetSlotWidgets() const
+const TMap<FName, FPreviewableWidgetVariant>& UWidgetPreview::GetSlotWidgetTypes() const
 {
-	return SlotWidgets;
+	return SlotWidgetTypes;
 }
 
-void UWidgetPreview::SetSlotWidgets(const TMap<FName, TSoftClassPtr<UUserWidget>>& InWidgets)
+void UWidgetPreview::SetSlotWidgetTypes(const TMap<FName, FPreviewableWidgetVariant>& InWidgets)
 {
-	if (!SlotWidgets.OrderIndependentCompareEqual(InWidgets))
+	if (!SlotWidgetTypes.OrderIndependentCompareEqual(InWidgets))
 	{
-		SlotWidgets = InWidgets;
+		SlotWidgetTypes = InWidgets;
 		WidgetInstance = nullptr;
 		UpdateWidgets();
 
@@ -273,10 +334,9 @@ void UWidgetPreview::PostEditChangeProperty(FPropertyChangedEvent& PropertyChang
 {
 	Super::PostEditChangeProperty(PropertyChangedEvent);
 
-	FName PropertyName = PropertyChangedEvent.GetPropertyName();
+	FName PropertyName = PropertyChangedEvent.GetMemberPropertyName();
 	if (PropertyName == GET_MEMBER_NAME_CHECKED(UWidgetPreview, WidgetType)
-		|| PropertyName == GET_MEMBER_NAME_CHECKED(UWidgetPreview, LayoutWidgetType)
-		|| PropertyName == GET_MEMBER_NAME_CHECKED(UWidgetPreview, SlotWidgets)
+		|| PropertyName == GET_MEMBER_NAME_CHECKED(UWidgetPreview, SlotWidgetTypes)
 		|| PropertyName.IsNone()) // None can be an Undo operation
 	{
 		WidgetInstance = nullptr;
@@ -285,49 +345,88 @@ void UWidgetPreview::PostEditChangeProperty(FPropertyChangedEvent& PropertyChang
 	}
 }
 
+void UWidgetPreview::PostEditChangeChainProperty(FPropertyChangedChainEvent& PropertyChangedEvent)
+{
+	UObject::PostEditChangeChainProperty(PropertyChangedEvent);
+}
+
 void UWidgetPreview::OnWidgetBlueprintChanged(UBlueprint* InBlueprint)
 {
+	WidgetInstance = nullptr;
+	UpdateWidgets();
 	OnWidgetChanged().Broadcast(EWidgetPreviewWidgetChangeType::Structure);
 }
 
 void UWidgetPreview::UpdateWidgets()
 {
-	auto AddOnChanged = [this](const UUserWidget* InUserWidget)
+	auto AddOnChanged = [this](const UUserWidget* InUserWidgetCDO)
 	{
-		const UE::UMGWidgetPreview::Private::FWidgetTypeTuple WidgetTuple(InUserWidget);
+		const UE::UMGWidgetPreview::Private::FWidgetTypeTuple WidgetTuple(InUserWidgetCDO);
 		if (UWidgetBlueprint* BP = WidgetTuple.Blueprint)
 		{
 			BP->OnChanged().AddUObject(this, &UWidgetPreview::OnWidgetBlueprintChanged);
 		}
 	};
 
-	if (!LayoutWidgetType.IsNull())
+	CleanupReferences();
+
+	WidgetType.UpdateCachedWidget();
+	if (!WidgetType.ObjectPath.IsNull())
 	{
 		SlotNameCache.Reset();
+		if (const UUserWidget* AsUserWidget = WidgetType.AsUserWidgetCDO())
+		{
+			WidgetReferenceCache.Emplace(MakeWeakObjectPtr(AsUserWidget));
+			AddOnChanged(AsUserWidget);
 
-		AddOnChanged(LayoutWidgetType.LoadSynchronous()->GetDefaultObject<UUserWidget>());
+			if (const INamedSlotInterface* WidgetWithSlots = Cast<INamedSlotInterface>(AsUserWidget))
+			{
+				TArray<FName> SlotNames;
+				WidgetWithSlots->GetSlotNames(SlotNames);
 
-		for (const TPair<FName, TSoftClassPtr<UUserWidget>>& SlotWidget : SlotWidgets)
+				SlotNameCache = SlotNames;
+			}
+		}
+
+		for (TPair<FName, FPreviewableWidgetVariant>& SlotWidget : SlotWidgetTypes)
 		{
-			if (!SlotWidget.Value.IsNull())
+			SlotWidget.Value.UpdateCachedWidget();
+			if (!SlotWidget.Value.ObjectPath.IsNull())
 			{
-				AddOnChanged(SlotWidget.Value.LoadSynchronous()->GetDefaultObject<UUserWidget>());
+				if (const UUserWidget* AsUserWidget = WidgetType.AsUserWidgetCDO())
+				{
+					WidgetReferenceCache.Emplace(MakeWeakObjectPtr(AsUserWidget));
+					AddOnChanged(AsUserWidget);
+				}
 			}
 		}
 	}
+}
 
-	if (!WidgetType.IsNull())
+void UWidgetPreview::CleanupReferences()
+{
+	// Clear previous references, required due to how Blueprints are handled when changed
+	TArray<TWeakObjectPtr<const UUserWidget>> WidgetsToDeinitialize = WidgetReferenceCache;
+	for (TWeakObjectPtr<const UUserWidget>& WeakUserWidget : WidgetsToDeinitialize)
 	{
-		AddOnChanged(WidgetType.LoadSynchronous()->GetDefaultObject<UUserWidget>());
+		if (const UUserWidget* UserWidget = WeakUserWidget.Get())
+		{
+			const UE::UMGWidgetPreview::Private::FWidgetTypeTuple WidgetTuple(UserWidget);
+			if (UWidgetBlueprint* BP = WidgetTuple.Blueprint)
+			{
+				BP->OnChanged().RemoveAll(this);
+			}
+		}
 	}
+	WidgetReferenceCache.Reset();
 }
 
-TArray<FName> UWidgetPreview::GetAvailableLayoutSlotNames()
+TArray<FName> UWidgetPreview::GetAvailableWidgetSlotNames()
 {
-	const TArray<FName> AllSlotNames = GetLayoutSlotNames();
+	const TArray<FName> AllSlotNames = GetWidgetSlotNames();
 
 	TArray<FName> UsedSlotNamesArray;
-	SlotWidgets.GenerateKeyArray(UsedSlotNamesArray);
+	SlotWidgetTypes.GenerateKeyArray(UsedSlotNamesArray);
 
 	const TSet<FName> UsedSlotNames(UsedSlotNamesArray);
 
diff --git a/Engine/Plugins/Editor/UMGWidgetPreview/Source/UMGWidgetPreview/Private/WidgetPreview.h b/Engine/Plugins/Editor/UMGWidgetPreview/Source/UMGWidgetPreview/Private/WidgetPreview.h
--- a/Engine/Plugins/Editor/UMGWidgetPreview/Source/UMGWidgetPreview/Private/WidgetPreview.h
+++ b/Engine/Plugins/Editor/UMGWidgetPreview/Source/UMGWidgetPreview/Private/WidgetPreview.h
@@ -6,10 +6,12 @@
 
 #include "WidgetPreview.generated.h"
 
+class UPanelWidget;
+class UUserWidget;
+class UWidget;
 class UWidgetBlueprint;
+class UWidgetPreview;
 struct FImage;
-class UWidget;
-class UUserWidget;
 
 enum class EWidgetPreviewWidgetChangeType
 {
@@ -19,7 +21,45 @@ enum class EWidgetPreviewWidgetChangeType
 	ChildReference = 3,
 };
 
-// @todo: allow nested previews
+USTRUCT(BlueprintType)
+struct UMGWIDGETPREVIEW_API FPreviewableWidgetVariant
+{
+	GENERATED_BODY();
+
+	UPROPERTY(BlueprintReadOnly, EditAnywhere, Category = "Widget", DisplayName = "Widget Type", meta = (AllowedClasses = "/Script/UMGEditor.WidgetBlueprint, /Script/UMGWidgetPreview.WidgetPreview"))
+	FSoftObjectPath ObjectPath;
+
+	FPreviewableWidgetVariant() = default;
+	FPreviewableWidgetVariant(const TSubclassOf<UUserWidget>& InWidgetType);
+	FPreviewableWidgetVariant(const UWidgetPreview* InWidgetPreview);
+
+public:
+	/** Flushes cached widgets and re-resolves from the ObjectPath. */
+	void UpdateCachedWidget();
+
+	/** Returns the referenced Object as a UUserWidget (CDO). Returns nullptr if not found, or we couldn't find a nested UUserWidget (ie. inside a UWidgetPreview). */
+	const UUserWidget* AsUserWidgetCDO() const;
+
+	/** Returns the referenced Object as a UWidgetPreview. Returns nullptr if not found, or not a UWidgetPreview. */
+	UWidgetPreview* AsWidgetPreview() const;
+
+	friend bool operator==(const FPreviewableWidgetVariant& Left, const FPreviewableWidgetVariant& Right)
+	{
+		return Left.ObjectPath == Right.ObjectPath;
+	}
+
+	friend bool operator!=(const FPreviewableWidgetVariant& Left, const FPreviewableWidgetVariant& Right)
+	{
+		return !(Left == Right);
+	}
+
+private:
+	UPROPERTY(Transient)
+	TWeakObjectPtr<const UUserWidget> CachedWidgetCDO;
+
+	UPROPERTY(Transient)
+	TWeakObjectPtr<UWidgetPreview> CachedWidgetPreview;
+};
 
 UCLASS(BlueprintType, NotBlueprintable, AutoExpandCategories = "Widgets")
 class UMGWIDGETPREVIEW_API UWidgetPreview
@@ -28,14 +68,7 @@ class UMGWIDGETPREVIEW_API UWidgetPreview
 	GENERATED_BODY()
 
 public:
-	UFUNCTION(BlueprintCallable, Category = "Layout")
-	const TArray<FName>& GetLayoutSlotNames();
-
-	/** Returns or builds and returns an instance of the root widget for previewing. Can be used to trigger a rebuild. */
-	[[maybe_unused]] UUserWidget* GetOrCreateWidgetInstance(UWorld* InWorld, const bool bInForceRecreate = false);
-
-	UUserWidget* GetWidget() const;
-	const UUserWidget* GetWidgetForSlot(const FName InSlotName) const;
+	virtual void BeginDestroy() override;
 
 	/** Convenience function to check that all utilized widgets have bCanCallInitializedWithoutPlayerContext set to true, and reports any that don't. */
 	bool CanCallInitializedWithoutPlayerContext(const bool bInRecursive, TArray<const UUserWidget*>& OutFailedWidgets);
@@ -48,49 +81,61 @@ public:
 	using FOnWidgetChanged = TMulticastDelegate<void(const EWidgetPreviewWidgetChangeType)>;
 
 	FOnWidgetChanged& OnWidgetChanged() { return OnWidgetChangedDelegate; }
-	
-	const TSoftClassPtr<UUserWidget>& GetWidgetType() const;
-	void SetWidgetType(const TSoftClassPtr<UUserWidget>& InWidget);
-	
-	const TSoftClassPtr<UUserWidget>& GetLayoutWidgetType() const;
-	void SetLayoutWidgetType(const TSoftClassPtr<UUserWidget>& InWidget);
 
-	const TMap<FName, TSoftClassPtr<UUserWidget>>& GetSlotWidgets() const;
-	void SetSlotWidgets(const TMap<FName, TSoftClassPtr<UUserWidget>>& InWidgets);
+	UFUNCTION(BlueprintCallable, Category = "Layout")
+	const TArray<FName>& GetWidgetSlotNames() const;
+
+	/** Returns or builds and returns an instance of the root widget for previewing. Can be used to trigger a rebuild. */
+	[[maybe_unused]] UUserWidget* GetOrCreateWidgetInstance(UWorld* InWorld, const bool bInForceRecreate = false);
+
+	/** Returns the current widget instance, if any. */
+	UUserWidget* GetWidgetInstance() const;
+
+	const UUserWidget* GetWidgetCDO() const;
+	const UUserWidget* GetWidgetCDOForSlot(const FName InSlotName) const;
+
+	const FPreviewableWidgetVariant& GetWidgetType() const;
+	void SetWidgetType(const FPreviewableWidgetVariant& InWidget);
+
+	const TMap<FName, FPreviewableWidgetVariant>& GetSlotWidgetTypes() const;
+	void SetSlotWidgetTypes(const TMap<FName, FPreviewableWidgetVariant>& InWidgets);
 
 protected:
 	virtual void PostLoad() override;
 
 	virtual void PostEditChangeProperty(FPropertyChangedEvent& PropertyChangedEvent) override;
+	virtual void PostEditChangeChainProperty(FPropertyChangedChainEvent& PropertyChangedEvent) override;
 
 	void OnWidgetBlueprintChanged(UBlueprint* InBlueprint);
 
 	/** Misc. functionality to perform after a widget assignment is changed. */
 	void UpdateWidgets();
 
+	void CleanupReferences();
+
 	/** Returns slot names not already occupied in SlotWidgets. */
-	UFUNCTION(BlueprintCallable, Category = "Layout")
-	TArray<FName> GetAvailableLayoutSlotNames();
+	UFUNCTION(BlueprintCallable, Category = "Widget")
+	TArray<FName> GetAvailableWidgetSlotNames();
 
 private:
-	/** Widget to use. Will be overlaid on the Layout Widget (if specified). */
-	UPROPERTY(EditAnywhere, BlueprintReadWrite, Getter, Setter, Category = "Widget", DisplayName = "Widget", meta = (AllowPrivateAccess = "true"))
-	TSoftClassPtr<UUserWidget> WidgetType;
-
-	/** Optional widget to use for layout, containing named slots. */
-	UPROPERTY(EditAnywhere, BlueprintReadWrite, Getter, Setter, Category = "Layout", DisplayName = "Layout Widget", meta = (AllowPrivateAccess = "true"))
-	TSoftClassPtr<UUserWidget> LayoutWidgetType;
+	/** Widget to preview. */
+	UPROPERTY(EditAnywhere, BlueprintReadWrite, Getter, Setter, Category = "Widget", DisplayName = "Widget", meta = (AllowPrivateAccess = "true", ShowOnlyInnerProperties))
+	FPreviewableWidgetVariant WidgetType;
 
-	/** Widget per-slot. */
-	UPROPERTY(EditAnywhere, BlueprintReadWrite, Getter, Setter, Category = "Layout", meta = (AllowPrivateAccess = "true", EditCondition = "LayoutWidgetType", GetKeyOptions = "GetAvailableLayoutSlotNames"))
-	TMap<FName, TSoftClassPtr<UUserWidget>> SlotWidgets;
+	/** Widget per-slot, if WidgetType has any. */
+	UPROPERTY(EditAnywhere, BlueprintReadWrite, Getter, Setter, Category = "Widget", DisplayName = "Slot Widgets", meta = (AllowPrivateAccess = "true", GetKeyOptions = "GetAvailableWidgetSlotNames", ShowOnlyInnerProperties))
+	TMap<FName, FPreviewableWidgetVariant> SlotWidgetTypes;
 
 	UPROPERTY(Transient)
 	TObjectPtr<UUserWidget> WidgetInstance;
 
-	/** Slot names available in LayoutWidgetType. */
+	/** Slot names available in WidgetType (if any). */
 	UPROPERTY(Transient)
 	TArray<FName> SlotNameCache;
 
+	/** Widgets here should be checked for validity when a new one is assigned, to allow tear-down functionality. */
+	UPROPERTY(Transient)
+	TArray<TWeakObjectPtr<const UUserWidget>> WidgetReferenceCache;
+
 	FOnWidgetChanged OnWidgetChangedDelegate;
 };
diff --git a/Engine/Plugins/Editor/UMGWidgetPreview/Source/UMGWidgetPreview/Private/WidgetPreviewTypesPrivate.h b/Engine/Plugins/Editor/UMGWidgetPreview/Source/UMGWidgetPreview/Private/WidgetPreviewTypesPrivate.h
--- a/Engine/Plugins/Editor/UMGWidgetPreview/Source/UMGWidgetPreview/Private/WidgetPreviewTypesPrivate.h
+++ b/Engine/Plugins/Editor/UMGWidgetPreview/Source/UMGWidgetPreview/Private/WidgetPreviewTypesPrivate.h
@@ -10,11 +10,11 @@ namespace UE::UMGWidgetPreview::Private
 {
 	struct FWidgetTypeTuple
 	{
-		explicit FWidgetTypeTuple(const UUserWidget* InUserWidget)
+		explicit FWidgetTypeTuple(const UUserWidget* InUserWidgetCDO)
 		{
-			check(InUserWidget);
+			check(InUserWidgetCDO);
 
-			ClassDefaultObject = InUserWidget;
+			ClassDefaultObject = InUserWidgetCDO;
 			BlueprintGeneratedClass = Cast<UWidgetBlueprintGeneratedClass>(ClassDefaultObject->GetClass());
 			Blueprint = Cast<UWidgetBlueprint>(BlueprintGeneratedClass->ClassGeneratedBy);
 		}
