[UMG Widget Preview] CIS fix

diff --git a/Engine/Plugins/Editor/UMGWidgetPreview/Source/UMGWidgetPreview/Private/WidgetPreviewEditor.cpp b/Engine/Plugins/Editor/UMGWidgetPreview/Source/UMGWidgetPreview/Private/WidgetPreviewEditor.cpp
--- a/Engine/Plugins/Editor/UMGWidgetPreview/Source/UMGWidgetPreview/Private/WidgetPreviewEditor.cpp
+++ b/Engine/Plugins/Editor/UMGWidgetPreview/Source/UMGWidgetPreview/Private/WidgetPreviewEditor.cpp
@@ -99,7 +99,7 @@ UWidgetPreview* UWidgetPreviewEditor::CreatePreviewForWidget(const UUserWidget*
 			nullptr,
 			GWarn));
 
-	NewWidgetPreview->SetWidgetType(InUserWidget->GetClass());
+	NewWidgetPreview->SetWidgetType({ InUserWidget->GetClass() });
 	NewWidgetPreview->MarkPackageDirty();
 
 	return NewWidgetPreview;
