[UMG Widget Preview] Quick preview no longer saves a package by default - a user can optionally save it

diff --git a/Engine/Plugins/Editor/UMGWidgetPreview/Source/UMGWidgetPreview/Private/WidgetPreviewEditor.cpp b/Engine/Plugins/Editor/UMGWidgetPreview/Source/UMGWidgetPreview/Private/WidgetPreviewEditor.cpp
--- a/Engine/Plugins/Editor/UMGWidgetPreview/Source/UMGWidgetPreview/Private/WidgetPreviewEditor.cpp
+++ b/Engine/Plugins/Editor/UMGWidgetPreview/Source/UMGWidgetPreview/Private/WidgetPreviewEditor.cpp
@@ -86,14 +86,11 @@ UWidgetPreview* UWidgetPreviewEditor::CreatePreviewForWidget(const UUserWidget*
 	FString PackageName;
 	AssetToolsModule.Get().CreateUniqueAssetName(DesiredPackagePath, DefaultSuffix, PackageName, AssetName);
 
-	UPackage* Package = CreatePackage(*PackageName);
-	Package->SetPackageFlags(PKG_NewlyCreated);
-
 	EObjectFlags Flags = RF_Public | RF_Standalone;
 	UWidgetPreview* NewWidgetPreview = CastChecked<UWidgetPreview>(
 		Factory->FactoryCreateNew(
 			UWidgetPreview::StaticClass(),
-			Package,
+			GetTransientPackage(),
 			*AssetName,
 			Flags,
 			nullptr,
diff --git a/Engine/Plugins/Editor/UMGWidgetPreview/Source/UMGWidgetPreview/Private/WidgetPreviewToolkit.cpp b/Engine/Plugins/Editor/UMGWidgetPreview/Source/UMGWidgetPreview/Private/WidgetPreviewToolkit.cpp
--- a/Engine/Plugins/Editor/UMGWidgetPreview/Source/UMGWidgetPreview/Private/WidgetPreviewToolkit.cpp
+++ b/Engine/Plugins/Editor/UMGWidgetPreview/Source/UMGWidgetPreview/Private/WidgetPreviewToolkit.cpp
@@ -3,6 +3,7 @@
 #include "WidgetPreviewToolkit.h"
 
 #include "DataValidationFixers.h"
+#include "FileHelpers.h"
 #include "Framework/Application/SlateApplication.h"
 #include "Kismet2/KismetEditorUtilities.h"
 #include "MessageLogModule.h"
@@ -177,6 +178,16 @@ namespace UE::UMGWidgetPreview::Private
 		StatusMessage = FTokenizedMessage::Create(EMessageSeverity::Info, LOCTEXT("WidgetPreviewToolkitBackgroundState_Message", "The widget preview is paused while the window is in the background. Re-focus to unpause."));
 	}
 
+	void FWidgetPreviewToolkitBackgroundState::OnEnter(const FWidgetPreviewToolkitStateBase* InFromState)
+	{
+		FWidgetPreviewToolkitPausedState::OnEnter(InFromState);
+	}
+
+	void FWidgetPreviewToolkitBackgroundState::OnExit(const FWidgetPreviewToolkitStateBase* InToState)
+	{
+		FWidgetPreviewToolkitPausedState::OnExit(InToState);
+	}
+
 	FWidgetPreviewToolkitUnsupportedWidgetState::FWidgetPreviewToolkitUnsupportedWidgetState()
 	{
 		Id = TEXT("UnsupportedWidget");
@@ -305,18 +316,22 @@ namespace UE::UMGWidgetPreview::Private
 
 	void FWidgetPreviewToolkit::RegisterToolbar()
 	{
+		FToolMenuOwnerScoped ToolMenuOwnerScope(this);
+
 		FName ParentName;
 		const FName MenuName = GetToolMenuToolbarName(ParentName);
 		UToolMenus* ToolMenus = UToolMenus::Get();
+		UToolMenu* ToolbarMenu = ToolMenus->ExtendMenu(MenuName);
 		if (!ToolMenus->IsMenuRegistered(MenuName))
 		{
-			FToolMenuOwnerScoped ToolMenuOwnerScope(this);
+			ToolbarMenu = ToolMenus->RegisterMenu(MenuName, ParentName, EMultiBoxType::ToolBar);
+		}
 
-			UToolMenu* ToolbarMenu = ToolMenus->RegisterMenu(MenuName, ParentName, EMultiBoxType::ToolBar);
-
-			FToolMenuInsert InsertAfterAssetSection("Asset", EToolMenuInsertType::After);
-			const FWidgetPreviewCommands& Commands = FWidgetPreviewCommands::Get();
+		FToolMenuInsert InsertAfterAssetSection("Asset", EToolMenuInsertType::After);
+		const FWidgetPreviewCommands& Commands = FWidgetPreviewCommands::Get();
 
+		// Preview Section
+		{
 			FToolMenuSection& PreviewSection = ToolbarMenu->FindOrAddSection(
 				"Preview",
 				{},
@@ -383,6 +398,146 @@ namespace UE::UMGWidgetPreview::Private
 		ResolveState();
 	}
 
+	bool FWidgetPreviewToolkit::CanSaveAsset() const
+	{
+		// We use the same logic here - if the outer package is transient, the only option is "SaveAs"
+		return IsSaveAssetVisible();
+	}
+
+	void FWidgetPreviewToolkit::SaveAsset_Execute()
+	{
+		TArray<UObject*> ObjectsToSave;
+		GetSaveableObjects(ObjectsToSave);
+
+		if (ObjectsToSave.Num() == 0)
+		{
+			return;
+		}
+
+		TArray<UObject*> SavedObjects;
+		SavedObjects.Reserve(ObjectsToSave.Num());
+
+		TArray<UPackage*> PackagesToSave;
+		bool bHasNewlyCreatedPackage = false;
+
+		TArray<UPackage*> PackagesToSaveAs;
+
+		for (UObject* Object : ObjectsToSave)
+		{
+			if (Object == nullptr)
+			{
+				// Log an invalid object but don't try to save it
+				UE_LOG(LogWidgetPreview, Log, TEXT("Invalid preview to save: %s"), (Object != nullptr) ? *Object->GetFullName() : TEXT("Null Object"));
+			}
+			else
+			{
+				SavedObjects.Add(Object);
+			}
+		}
+
+		FEditorFileUtils::PromptForCheckoutAndSave(PackagesToSave, bCheckDirtyOnAssetSave, /*bPromptToSave=*/ bHasNewlyCreatedPackage);
+
+		OnAssetsSaved(SavedObjects);
+	}
+
+	bool FWidgetPreviewToolkit::IsSaveAssetAsVisible() const
+	{
+		// @note: usually this wouldn't appear when the asset belongs to the transient package
+		// We allow this so that the user has the option of saving it to an asset (non-transient package)
+		return true;
+	}
+
+	void FWidgetPreviewToolkit::SaveAssetAs_Execute()
+	{
+		TSharedPtr<IToolkitHost> MyToolkitHost = ToolkitHost.Pin();
+		if (!MyToolkitHost.IsValid())
+		{
+			return;
+		}
+
+		TArray<UObject*> ObjectsToSave;
+		GetSaveableObjects(ObjectsToSave);
+
+		if (ObjectsToSave.Num() == 0)
+		{
+			return;
+		}
+
+		TArray<UObject*> ObjectsToSaveWithoutPackage;
+		ObjectsToSaveWithoutPackage.Reserve(ObjectsToSave.Num());
+
+		// Temporarily set to Transient objects, so SaveAssetsAs will auto-populate the default path
+		for (UObject* Object : ObjectsToSave)
+		{
+			UPackage* Package = Object->GetPackage();
+			if (!Package || Package == GetTransientPackage())
+			{
+				Object->SetFlags(Object->GetFlags() | RF_Transient);
+				ObjectsToSaveWithoutPackage.Emplace(Object);
+			}
+		}
+
+		TArray<UObject*> SavedObjects;
+		FEditorFileUtils::SaveAssetsAs(ObjectsToSave, SavedObjects);
+
+		if (SavedObjects.Num() == 0)
+		{
+			// Error saving, or user closed the dialog. Restore objects to non-transient
+			for (UObject* Object : ObjectsToSaveWithoutPackage)
+            {
+				Object->ClearFlags(EObjectFlags::RF_Transient);
+            }
+
+			return;
+		}
+
+		// close existing asset editors for resaved assets
+		UAssetEditorSubsystem* AssetEditorSubsystem = GEditor->GetEditorSubsystem<UAssetEditorSubsystem>();
+
+		TArray<UObject*> ObjectsBeingEdited = GetEditingObjects();
+
+		// hack @see: FAssetEditorToolkit::SaveAssetAs_Execute()
+		TArray<UObject*> ObjectsToReopen;
+		for (UObject* Object : ObjectsBeingEdited)
+		{
+			if (Object->IsAsset() && !ObjectsToSave.Contains(Object))
+			{
+				ObjectsToReopen.Add(Object);
+			}
+		}
+
+		for (UObject* Object : SavedObjects)
+		{
+			if (ShouldReopenEditorForSavedAsset(Object))
+			{
+				ObjectsToReopen.AddUnique(Object);
+			}
+		}
+
+		for (UObject* Object : ObjectsBeingEdited)
+		{
+			AssetEditorSubsystem->CloseAllEditorsForAsset(Object);
+			GEditor->GetEditorSubsystem<UAssetEditorSubsystem>()->NotifyAssetClosed(Object, this);
+		}
+
+		AssetEditorSubsystem->OpenEditorForAssets_Advanced(ObjectsToReopen, ToolkitMode, MyToolkitHost.ToSharedRef());
+		// end hack
+
+		OnAssetsSavedAs(SavedObjects);
+	}
+
+	void FWidgetPreviewToolkit::GetSaveableObjects(TArray<UObject*>& OutObjects) const
+	{
+		FBaseAssetToolkit::GetSaveableObjects(OutObjects);
+
+		TArray<UObject*> ObjectsBeingEdited = GetEditingObjects();
+		for (const TObjectPtr<UObject>& Object : ObjectsBeingEdited)
+		{
+			// We override this to allow Transient objects to be saved
+			OutObjects.Add(Object);
+		}
+	}
+
 	FText FWidgetPreviewToolkit::GetToolkitName() const
 	{
 		const TArray<UObject*>* Objects = GetObjectsCurrentlyBeingEdited();
@@ -409,7 +564,7 @@ namespace UE::UMGWidgetPreview::Private
 
 	FString FWidgetPreviewToolkit::GetWorldCentricTabPrefix() const
 	{
-		return LOCTEXT("WorldCentricTabPrefix", "WidgetPreviw").ToString();
+		return LOCTEXT("WorldCentricTabPrefix", "WidgetPreview").ToString();
 	}
 
 	FLinearColor FWidgetPreviewToolkit::GetWorldCentricTabColorScale() const
@@ -442,16 +597,6 @@ namespace UE::UMGWidgetPreview::Private
 		return OnStateChangedDelegate;
 	}
 
-	TSharedPtr<FWidgetBlueprintEditor> FWidgetPreviewToolkit::GetEditor() const
-	{
-		if (TSharedPtr<FWidgetBlueprintEditor> Editor = WeakEditor.Pin())
-		{
-			return Editor;
-		}
-
-		return nullptr;
-	}
-
 	TSharedPtr<FWidgetPreviewScene> FWidgetPreviewToolkit::GetPreviewScene()
 	{
 		if (!PreviewScene.IsValid())
diff --git a/Engine/Plugins/Editor/UMGWidgetPreview/Source/UMGWidgetPreview/Private/WidgetPreviewToolkit.h b/Engine/Plugins/Editor/UMGWidgetPreview/Source/UMGWidgetPreview/Private/WidgetPreviewToolkit.h
--- a/Engine/Plugins/Editor/UMGWidgetPreview/Source/UMGWidgetPreview/Private/WidgetPreviewToolkit.h
+++ b/Engine/Plugins/Editor/UMGWidgetPreview/Source/UMGWidgetPreview/Private/WidgetPreviewToolkit.h
@@ -34,7 +34,7 @@ namespace UE::UMGWidgetPreview::Private
 
 	/** Encapsulates the state needed to run the preview world. */
 	class FWidgetPreviewScene
-	: public FTickableEditorObject
+		: public FTickableEditorObject
 	{
 	public:
 		explicit FWidgetPreviewScene(const TSharedRef<FWidgetPreviewToolkit>& InPreviewToolkit);
@@ -88,6 +88,9 @@ namespace UE::UMGWidgetPreview::Private
 	struct FWidgetPreviewToolkitBackgroundState : FWidgetPreviewToolkitPausedState
 	{
 		FWidgetPreviewToolkitBackgroundState();
+
+		virtual void OnEnter(const FWidgetPreviewToolkitStateBase* InFromState) override;
+		virtual void OnExit(const FWidgetPreviewToolkitStateBase* InToState) override;
 	};
 
 	struct FWidgetPreviewToolkitUnsupportedWidgetState : FWidgetPreviewToolkitPausedState
@@ -122,9 +125,9 @@ namespace UE::UMGWidgetPreview::Private
 	};
 
 	class FWidgetPreviewToolkit
-	: public FBaseAssetToolkit
-	  , public FGCObject
-	  , public IWidgetPreviewToolkit
+		: public FBaseAssetToolkit
+		, public FGCObject
+		, public IWidgetPreviewToolkit
 	{
 	public:
 		explicit FWidgetPreviewToolkit(UWidgetPreviewEditor* InOwningEditor);
@@ -139,6 +142,11 @@ namespace UE::UMGWidgetPreview::Private
 		virtual void RegisterTabSpawners(const TSharedRef<FTabManager>& TabManager) override;
 		virtual void UnregisterTabSpawners(const TSharedRef<FTabManager>& InTabManager) override;
 		virtual void PostInitAssetEditor() override;
+		virtual bool CanSaveAsset() const override;
+		virtual void SaveAsset_Execute() override;
+		virtual bool IsSaveAssetAsVisible() const override;
+		virtual void SaveAssetAs_Execute() override;
+		virtual void GetSaveableObjects(TArray<UObject*>& OutObjects) const override;
 		//~ End FAssetEditorToolkit
 
 		//~ Begin IToolkit
@@ -160,7 +168,6 @@ namespace UE::UMGWidgetPreview::Private
 		virtual FOnStateChanged& OnStateChanged() override;
 		//~ End IWidgetPreviewToolkit
 
-		TSharedPtr<FWidgetBlueprintEditor> GetEditor() const;
 		TSharedPtr<FWidgetPreviewScene> GetPreviewScene();
 		UWorld* GetPreviewWorld();
 
@@ -192,7 +199,6 @@ namespace UE::UMGWidgetPreview::Private
 		TSharedRef<SDockTab> SpawnTab_MessageLog(const FSpawnTabArgs& Args);
 
 	private:
-		TWeakPtr<FWidgetBlueprintEditor> WeakEditor;
 		TObjectPtr<UWidgetPreview> Preview;
 
 		TSharedPtr<FWidgetPreviewScene> PreviewScene;
