[UMG Widget Preview] CIS fix

diff --git a/Engine/Plugins/Editor/UMGWidgetPreview/Source/UMGWidgetPreview/Private/WidgetPreviewToolkit.cpp b/Engine/Plugins/Editor/UMGWidgetPreview/Source/UMGWidgetPreview/Private/WidgetPreviewToolkit.cpp
--- a/Engine/Plugins/Editor/UMGWidgetPreview/Source/UMGWidgetPreview/Private/WidgetPreviewToolkit.cpp
+++ b/Engine/Plugins/Editor/UMGWidgetPreview/Source/UMGWidgetPreview/Private/WidgetPreviewToolkit.cpp
@@ -531,7 +531,7 @@ namespace UE::UMGWidgetPreview::Private
 		FBaseAssetToolkit::GetSaveableObjects(OutObjects);
 
 		TArray<UObject*> ObjectsBeingEdited = GetEditingObjects();
-		for (const TObjectPtr<UObject>& Object : ObjectsBeingEdited)
+		for (UObject* Object : ObjectsBeingEdited)
 		{
 			// We override this to allow Transient objects to be saved
 			OutObjects.Add(Object);
