[UMG Widget Preview] Switch to thumbnail when unfocused to prevent slate resource crash on BP change

diff --git a/Engine/Plugins/Editor/UMGWidgetPreview/Source/UMGWidgetPreview/Private/Widgets/SWidgetPreview.cpp b/Engine/Plugins/Editor/UMGWidgetPreview/Source/UMGWidgetPreview/Private/Widgets/SWidgetPreview.cpp
--- a/Engine/Plugins/Editor/UMGWidgetPreview/Source/UMGWidgetPreview/Private/Widgets/SWidgetPreview.cpp
+++ b/Engine/Plugins/Editor/UMGWidgetPreview/Source/UMGWidgetPreview/Private/Widgets/SWidgetPreview.cpp
@@ -19,6 +19,7 @@ namespace UE::UMGWidgetPreview::Private
 	{
 		WeakToolkit = InToolkit;
 
+		OnStateChangedHandle = InToolkit->OnStateChanged().AddSP(this, &SWidgetPreview::OnStateChanged);
 		OnWidgetChangedHandle = InToolkit->GetPreview()->OnWidgetChanged().AddSP(this, &SWidgetPreview::OnWidgetChanged);
 
 		CreatedSlateWidget = SNullWidget::NullWidget;
@@ -36,6 +37,8 @@ namespace UE::UMGWidgetPreview::Private
 			+ SOverlay::Slot()
 			[
 				SAssignNew(RetainerWidget, SRetainerWidget)
+				.RenderOnPhase(false)
+				.RenderOnInvalidation(false)
 				[
 					ContainerWidget.ToSharedRef()
 				]
@@ -44,18 +47,70 @@ namespace UE::UMGWidgetPreview::Private
 	}
 
 	SWidgetPreview::~SWidgetPreview()
-    {
-    	if (const TSharedPtr<FWidgetPreviewToolkit> Toolkit = WeakToolkit.Pin())
-    	{
-    		if (UWidgetPreview* Preview = Toolkit->GetPreview())
-    		{
-    			Preview->OnWidgetChanged().Remove(OnWidgetChangedHandle);
-    		}
-    	}
-    }
+	{
+		if (const TSharedPtr<FWidgetPreviewToolkit> Toolkit = WeakToolkit.Pin())
+		{
+			Toolkit->OnStateChanged().Remove(OnStateChangedHandle);
+
+			if (UWidgetPreview* Preview = Toolkit->GetPreview())
+			{
+				Preview->OnWidgetChanged().Remove(OnWidgetChangedHandle);
+			}
+		}
+	}
+
+	int32 SWidgetPreview::OnPaint(
+		const FPaintArgs& Args,
+		const FGeometry& AllottedGeometry,
+		const FSlateRect& MyCullingRect,
+		FSlateWindowElementList& OutDrawElements,
+		int32 LayerId,
+		const FWidgetStyle& InWidgetStyle,
+		bool bParentEnabled) const
+	{
+		const int32 Result = SCompoundWidget::OnPaint(
+			Args,
+			AllottedGeometry,
+			MyCullingRect,
+			OutDrawElements,
+			LayerId,
+			InWidgetStyle,
+			bParentEnabled);
+
+		if (bClearWidgetOnNextPaint)
+		{
+			SWidgetPreview* MutableThis = const_cast<SWidgetPreview*>(this);
+
+			MutableThis->CreatedSlateWidget = SNullWidget::NullWidget;
+			ContainerWidget->SetContent(CreatedSlateWidget.ToSharedRef());
+			MutableThis->bClearWidgetOnNextPaint = false;
+		}
+
+		return Result;
+	}
+
+	void SWidgetPreview::OnStateChanged(FWidgetPreviewToolkitStateBase* InOldState, FWidgetPreviewToolkitStateBase* InNewState)
+	{
+		const bool bShouldUseLiveWidget = InNewState->CanTick();
+		bIsRetainedRender = !bShouldUseLiveWidget;
+		bClearWidgetOnNextPaint = bIsRetainedRender;
+		RetainerWidget->RequestRender();
+		RetainerWidget->SetRetainedRendering(bIsRetainedRender);
+
+		if (bShouldUseLiveWidget)
+		{
+			OnWidgetChanged(EWidgetPreviewWidgetChangeType::Assignment);
+		}
+	}
 
 	void SWidgetPreview::OnWidgetChanged(const EWidgetPreviewWidgetChangeType InChangeType)
 	{
+		// Disallow widget assignment if retaining (cached thumbnail)
+		if (bIsRetainedRender)
+		{
+			return;
+		}
+
 		if (const TSharedPtr<FWidgetPreviewToolkit> Toolkit = WeakToolkit.Pin())
 		{
 			if (UWidgetPreview* Preview = Toolkit->GetPreview())
diff --git a/Engine/Plugins/Editor/UMGWidgetPreview/Source/UMGWidgetPreview/Private/Widgets/SWidgetPreview.h b/Engine/Plugins/Editor/UMGWidgetPreview/Source/UMGWidgetPreview/Private/Widgets/SWidgetPreview.h
--- a/Engine/Plugins/Editor/UMGWidgetPreview/Source/UMGWidgetPreview/Private/Widgets/SWidgetPreview.h
+++ b/Engine/Plugins/Editor/UMGWidgetPreview/Source/UMGWidgetPreview/Private/Widgets/SWidgetPreview.h
@@ -14,6 +14,7 @@ class SRetainerWidget;
 
 namespace UE::UMGWidgetPreview::Private
 {
+	struct FWidgetPreviewToolkitStateBase;
 	class FWidgetPreviewToolkit;
 
 	class SWidgetPreview
@@ -25,11 +26,14 @@ namespace UE::UMGWidgetPreview::Private
 		SLATE_BEGIN_ARGS(SWidgetPreview) {}
 		SLATE_END_ARGS()
 
+		void Construct(const FArguments& Args, const TSharedRef<FWidgetPreviewToolkit>& InToolkit);
+
 		virtual ~SWidgetPreview() override;
 
-		void Construct(const FArguments& Args, const TSharedRef<FWidgetPreviewToolkit>& InToolkit);
+		virtual int32 OnPaint(const FPaintArgs& Args, const FGeometry& AllottedGeometry, const FSlateRect& MyCullingRect, FSlateWindowElementList& OutDrawElements, int32 LayerId, const FWidgetStyle& InWidgetStyle, bool bParentEnabled) const override;
 
 	private:
+		void OnStateChanged(FWidgetPreviewToolkitStateBase* InOldState, FWidgetPreviewToolkitStateBase* InNewState);
 		void OnWidgetChanged(const EWidgetPreviewWidgetChangeType InChangeType);
 
 		/** Convenience method to get world from the associated viewport. */
@@ -42,6 +46,10 @@ namespace UE::UMGWidgetPreview::Private
 		TSharedPtr<SBorder> ContainerWidget;
 		TSharedPtr<SWidget> CreatedSlateWidget;
 
+		bool bClearWidgetOnNextPaint = false;
+		bool bIsRetainedRender = false;
+
+		FDelegateHandle OnStateChangedHandle;
 		FDelegateHandle OnWidgetChangedHandle;
 	};
 }
