[UMG Widget Preview] Exposed CallInEditor functions within the contained widgets to the UI

diff --git a/Engine/Plugins/Editor/UMGWidgetPreview/Source/UMGWidgetPreview/Private/Customizations/PreviewableWidgetCustomization.cpp b/Engine/Plugins/Editor/UMGWidgetPreview/Source/UMGWidgetPreview/Private/Customizations/PreviewableWidgetCustomization.cpp
new file mode 100644
--- /dev/null
+++ b/Engine/Plugins/Editor/UMGWidgetPreview/Source/UMGWidgetPreview/Private/Customizations/PreviewableWidgetCustomization.cpp
@@ -0,0 +1,374 @@
+// Copyright Epic Games, Inc. All Rights Reserved.
+
+#include "PreviewableWidgetCustomization.h"
+
+#include "Blueprint/UserWidget.h"
+#include "ContentBrowserModule.h"
+#include "Customizations/MathStructCustomizations.h"
+#include "DetailLayoutBuilder.h"
+#include "Framework/Application/SlateApplication.h"
+#include "Framework/MultiBox/MultiBoxBuilder.h"
+#include "Framework/Views/ITypedTableView.h"
+#include "IContentBrowserSingleton.h"
+#include "IDetailChildrenBuilder.h"
+#include "IDetailGroup.h"
+#include "IDetailPropertyRow.h"
+#include "IPropertyUtilities.h"
+#include "K2Node_CallFunction.h"
+#include "Kismet2/BlueprintEditorUtils.h"
+#include "Modules/ModuleManager.h"
+#include "ObjectTools.h"
+#include "PropertyCustomizationHelpers.h"
+#include "PropertyEditorModule.h"
+#include "Reflection/FunctionUtils.h"
+#include "ScopedTransaction.h"
+#include "Settings/BlueprintEditorProjectSettings.h"
+#include "Styling/StyleColors.h"
+#include "WidgetPreview.h"
+#include "Widgets/Images/SImage.h"
+#include "Widgets/Input/SButton.h"
+#include "Widgets/Input/SComboButton.h"
+#include "Widgets/Layout/SWrapBox.h"
+#include "Widgets/SBoxPanel.h"
+#include "Widgets/Text/STextBlock.h"
+#include "Widgets/Views/SListView.h"
+
+#define LOCTEXT_NAMESPACE "PreviewableWidgetCustomization"
+
+namespace UE::UMGWidgetPreview::Private
+{
+	TSharedRef<IPropertyTypeCustomization> FPreviewableWidgetCustomization::MakeInstance()
+	{
+		TSharedRef<FPreviewableWidgetCustomization> Customization = MakeShared<FPreviewableWidgetCustomization>();
+		return Customization;
+	}
+
+	void FPreviewableWidgetCustomization::CustomizeHeader(
+		TSharedRef<IPropertyHandle> PropertyHandle,
+		FDetailWidgetRow& HeaderRow,
+		IPropertyTypeCustomizationUtils& CustomizationUtils)
+	{
+		PropertyUtilities = CustomizationUtils.GetPropertyUtilities();
+		OwningPreviewHandle = PropertyHandle->GetParentHandle();
+		PreviewVariantHandle = PropertyHandle;
+		ObjectPathHandle = PreviewVariantHandle->GetChildHandle(GET_MEMBER_NAME_CHECKED(FPreviewableWidgetVariant, ObjectPath));
+
+		TArray<UObject*> OwningObjects;
+		PropertyHandle->GetOuterObjects(OwningObjects);
+		if (!OwningObjects.IsEmpty())
+		{
+			if (UWidgetPreview* OwningPreview = Cast<UWidgetPreview>(OwningObjects[0]))
+			{
+				OwningPreview->OnWidgetChanged().AddSP(this, &FPreviewableWidgetCustomization::OnWidgetChanged);
+				WeakOwningPreview = OwningPreview;
+			}
+		}
+	}
+
+	void FPreviewableWidgetCustomization::CustomizeChildren(
+		TSharedRef<IPropertyHandle> PropertyHandle,
+		IDetailChildrenBuilder& ChildBuilder,
+		IPropertyTypeCustomizationUtils& CustomizationUtils)
+	{
+		check(ObjectPathHandle.IsValid() && ObjectPathHandle->IsValidHandle());
+
+		ChildBuilder.AddProperty(ObjectPathHandle.ToSharedRef());
+
+		TArray<UFunction*> CallInEditorFunctions;
+		if (GetCallInEditorFunctions(CallInEditorFunctions))
+		{
+			AddCallInEditorFunctions(ChildBuilder, CallInEditorFunctions);
+		}
+	}
+
+	bool FPreviewableWidgetCustomization::GetCallInEditorFunctions(TArray<UFunction*>& OutCallInEditorFunctions) const
+	{
+		check(PreviewVariantHandle.IsValid() && PreviewVariantHandle->IsValidHandle());
+
+		FPreviewableWidgetVariant* PreviewableWidgetVariant = GetPreviewableWidgetVariant();
+		if (!PreviewableWidgetVariant
+			|| (PreviewableWidgetVariant && PreviewableWidgetVariant->ObjectPath.IsNull()))
+		{
+			return false;
+		}
+		
+		// metadata tag for defining sort order of function buttons within a Category
+		static const FName NAME_DisplayPriority("DisplayPriority");
+
+		const bool bDisallowEditorUtilityBlueprintFunctions = GetDefault<UBlueprintEditorProjectSettings>()->bDisallowEditorUtilityBlueprintFunctionsInDetailsView;
+
+		auto CanDisplayAndCallFunction = [](const UFunction* TestFunction)
+		{
+			bool bCanCall = TestFunction->GetBoolMetaData(FBlueprintMetadata::MD_CallInEditor)
+						&& TestFunction->HasAnyFunctionFlags(FUNC_Private)
+						&& TestFunction->ParmsSize == 0; // Params not supported
+
+			return bCanCall;
+		};
+
+		UClass* WidgetClass = PreviewableWidgetVariant->AsUserWidgetCDO()->GetClass();
+
+		// Get all of the functions we need to display (done ahead of time so we can sort them)
+		for (TFieldIterator<UFunction> FunctionIter(WidgetClass, EFieldIteratorFlags::IncludeSuper); FunctionIter; ++FunctionIter)
+		{
+			UFunction* TestFunction = *FunctionIter;
+			if (CanDisplayAndCallFunction(TestFunction))
+			{
+				bool bAllowFunction = true;
+				if (UClass* TestFunctionOwnerClass = TestFunction->GetOwnerClass())
+				{
+					if (UBlueprint* Blueprint = Cast<UBlueprint>(TestFunctionOwnerClass->ClassGeneratedBy))
+					{
+						if (FBlueprintEditorUtils::IsEditorUtilityBlueprint(Blueprint))
+						{
+							// Skip Blutilities if disabled via project settings
+							bAllowFunction = !bDisallowEditorUtilityBlueprintFunctions;
+						}
+					}
+				}
+
+				if (bAllowFunction)
+				{
+					const FName FunctionName = TestFunction->GetFName();
+					if (!OutCallInEditorFunctions.FindByPredicate([&FunctionName](const UFunction* Func) { return Func->GetFName() == FunctionName; }))
+					{
+						OutCallInEditorFunctions.Add(*FunctionIter);
+					}
+				}
+			}
+		}
+
+		if (!OutCallInEditorFunctions.IsEmpty())
+		{
+			// Sort the functions by category and then by DisplayPriority meta tag, and then by name
+			OutCallInEditorFunctions.Sort([](const UFunction& A, const UFunction& B)
+			{
+				const int32 CategorySort = A.GetMetaData(FBlueprintMetadata::MD_FunctionCategory).Compare(B.GetMetaData(FBlueprintMetadata::MD_FunctionCategory));
+				if (CategorySort != 0)
+				{
+					return (CategorySort <= 0);
+				}
+				else
+				{
+					FString DisplayPriorityAStr = A.GetMetaData(NAME_DisplayPriority);
+					int32 DisplayPriorityA = (DisplayPriorityAStr.IsEmpty() ? MAX_int32 : FCString::Atoi(*DisplayPriorityAStr));
+					if (DisplayPriorityA == 0 && !FCString::IsNumeric(*DisplayPriorityAStr))
+					{
+						DisplayPriorityA = MAX_int32;
+					}
+
+					FString DisplayPriorityBStr = B.GetMetaData(NAME_DisplayPriority);
+					int32 DisplayPriorityB = (DisplayPriorityBStr.IsEmpty() ? MAX_int32 : FCString::Atoi(*DisplayPriorityBStr));
+					if (DisplayPriorityB == 0 && !FCString::IsNumeric(*DisplayPriorityBStr))
+					{
+						DisplayPriorityB = MAX_int32;
+					}
+
+					return (DisplayPriorityA == DisplayPriorityB) ? (A.GetName() <= B.GetName()) : (DisplayPriorityA <= DisplayPriorityB);
+				}
+			});
+		}
+
+		// Return false if empty
+		return !OutCallInEditorFunctions.IsEmpty();
+	}
+
+	// Most of this is from FObjectDetails::AddCallInEditorMethods
+	void FPreviewableWidgetCustomization::AddCallInEditorFunctions(
+		IDetailChildrenBuilder& ChildBuilder,
+		const TArrayView<UFunction*>& InCallInEditorFunctions)
+	{
+		if (InCallInEditorFunctions.IsEmpty())
+		{
+			return;
+		}
+
+		UUserWidget* WidgetInstance = GetWidgetInstance();
+		if (!WidgetInstance)
+		{
+			return;
+		}
+
+		WeakWidgetInstance = WidgetInstance;
+
+		struct FCategoryEntry
+		{
+			FName CategoryName;
+			FName RowTag;
+			TSharedPtr<SWrapBox> WrapBox;
+			FTextBuilder FunctionSearchText;
+
+			FCategoryEntry(FName InCategoryName)
+				: CategoryName(InCategoryName)
+			{
+				WrapBox = SNew(SWrapBox)
+					// Setting the preferred size here (despite using UseAllottedSize) is a workaround for an issue
+					// when contained in a scroll box: prior to the first tick, the wrap box will use preferred size
+					// instead of allotted, and if preferred size is set small, it will cause the box to wrap a lot and
+					// request too much space from the scroll box. On next tick, SWrapBox is updated but the scroll box
+					// does not realize that it needs to show more elements, until it is scrolled.
+					// Setting a large value here means that the SWrapBox will request too little space prior to tick,
+					// which will cause the scroll box to virtualize more elements at the start, but this is less broken.
+					.PreferredSize(2000)
+					.UseAllottedSize(true);
+			}
+		};
+
+		// Build up a set of functions for each category, accumulating search text and buttons in a wrap box
+		FName ActiveCategory;
+		TArray<FCategoryEntry, TInlineAllocator<8>> CategoryList;
+		for (UFunction* Function : InCallInEditorFunctions)
+		{
+			if (!Function)
+			{
+				continue;
+			}
+
+			FName FunctionCategoryName(NAME_Default);
+			if (Function->HasMetaData(FBlueprintMetadata::FBlueprintMetadata::MD_FunctionCategory))
+			{
+				FunctionCategoryName = FName(*Function->GetMetaData(FBlueprintMetadata::MD_FunctionCategory));
+			}
+
+			if (FunctionCategoryName != ActiveCategory)
+			{
+				ActiveCategory = FunctionCategoryName;
+				CategoryList.Emplace(FunctionCategoryName);
+			}
+
+			FCategoryEntry& CategoryEntry = CategoryList.Last();
+
+			const FText ButtonCaption = ObjectTools::GetUserFacingFunctionName(Function);
+			FText FunctionTooltip = Function->GetToolTipText();
+			if (FunctionTooltip.IsEmpty())
+			{
+				FunctionTooltip = ButtonCaption;
+			}
+
+			TWeakObjectPtr<UFunction> WeakFunctionPtr(Function);
+			CategoryEntry.WrapBox->AddSlot()
+			.Padding(0.0f, 0.0f, 5.0f, 3.0f)
+			[
+				SNew(SButton)
+				.Text(ButtonCaption)
+				.OnClicked(FOnClicked::CreateSP(this, &FPreviewableWidgetCustomization::OnExecuteCallInEditorFunction, WeakFunctionPtr))
+				.IsEnabled(this, &FPreviewableWidgetCustomization::CanExecuteCallInEditorFunction, WeakFunctionPtr)
+				.ToolTipText(FunctionTooltip.IsEmptyOrWhitespace() ? LOCTEXT("CallInEditorTooltip", "Call an event on the selected object(s)") : FunctionTooltip)
+			];
+
+			CategoryEntry.RowTag = Function->GetFName();
+			CategoryEntry.FunctionSearchText.AppendLine(ButtonCaption);
+			CategoryEntry.FunctionSearchText.AppendLine(FunctionTooltip);
+
+			if (ButtonCaption.ToString() != Function->GetName())
+			{
+				CategoryEntry.FunctionSearchText.AppendLine(FText::FromString(Function->GetName()));
+			}
+		}
+
+		IDetailGroup& FunctionsGroup = ChildBuilder.AddGroup("Functions", LOCTEXT("FunctionsGroupName", "Functions"));
+		FunctionsGroup.ToggleExpansion(true);
+		FunctionsGroup.SetToolTip(LOCTEXT("FunctionsGroupToolTip", "CallInEditor functions within the referenced widget."));
+
+		TMap<FName, IDetailGroup*> Groups;
+
+		// Now edit the categories, adding the button strips to the details panel
+		for (FCategoryEntry& CategoryEntry : CategoryList)
+		{
+			IDetailGroup* Group = nullptr;
+			if (CategoryEntry.CategoryName == NAME_Default)
+			{
+				Group = &FunctionsGroup;
+			}
+			else if (IDetailGroup** ExistingGroup = Groups.Find(CategoryEntry.CategoryName))
+			{
+				Group = *ExistingGroup;
+			}
+			else
+			{
+				Group = Groups.Emplace(
+					CategoryEntry.CategoryName,
+					&FunctionsGroup.AddGroup(
+						CategoryEntry.CategoryName,
+						FText::FromName(CategoryEntry.CategoryName)));
+			}
+
+			Group->AddWidgetRow()
+			.FilterString(CategoryEntry.FunctionSearchText.ToText())
+			.ShouldAutoExpand(true)
+			.RowTag(CategoryEntry.RowTag)
+			[
+				CategoryEntry.WrapBox.ToSharedRef()
+			];
+		}
+	}
+
+	FReply FPreviewableWidgetCustomization::OnExecuteCallInEditorFunction(TWeakObjectPtr<UFunction> InWeakFunction)
+	{
+		using namespace UE::Reflection;
+		TArray<TWeakObjectPtr<UObject>> WeakExecutionObjects = GetFunctionCallExecutionContext(InWeakFunction);
+		if (UFunction* Function = InWeakFunction.Get())
+		{
+			FScopedTransaction Transaction(LOCTEXT("ExecuteCallInEditorMethod", "Call In Editor Action"));
+			if (!WeakExecutionObjects.IsEmpty())
+			{
+				TStrongObjectPtr<UFunction> StrongFunction(Function);
+	
+				FEditorScriptExecutionGuard ScriptGuard;
+				for (const TWeakObjectPtr<UObject>& WeakExecutionObject : WeakExecutionObjects)
+				{
+					if (UObject* ExecutionObject = WeakExecutionObject.Get())
+					{
+						ensure(Function->ParmsSize == 0);
+						TStrongObjectPtr<UObject> StrongExecutionObject(ExecutionObject); // Prevent GC during call
+						ExecutionObject->ProcessEvent(Function, nullptr);
+					}
+				}
+			}
+		}
+
+		return FReply::Handled();
+	}
+	
+	TArray<TWeakObjectPtr<UObject>> FPreviewableWidgetCustomization::GetFunctionCallExecutionContext(TWeakObjectPtr<UFunction> InWeakFunction) const
+	{
+		return { GetWidgetInstance() };
+	}
+
+	bool FPreviewableWidgetCustomization::CanExecuteCallInEditorFunction(TWeakObjectPtr<UFunction> InWeakFunction) const
+	{
+		return InWeakFunction.Get() && WeakWidgetInstance.IsValid();
+	}
+
+	FPreviewableWidgetVariant* FPreviewableWidgetCustomization::GetPreviewableWidgetVariant() const
+	{
+		check(PreviewVariantHandle.IsValid() && PreviewVariantHandle->IsValidHandle());
+
+		void* StructPtr = nullptr;
+		PreviewVariantHandle->GetValueData(StructPtr);
+		FPreviewableWidgetVariant* PreviewableWidgetVariant = reinterpret_cast<FPreviewableWidgetVariant*>(StructPtr);
+
+		return PreviewableWidgetVariant;
+	}
+
+	UUserWidget* FPreviewableWidgetCustomization::GetWidgetInstance() const
+	{
+		if (const UWidgetPreview* OwningPreview = WeakOwningPreview.Get())
+		{
+			// @todo: return instance of THIS widget, not the root
+			if (UUserWidget* WidgetInstance = OwningPreview->GetWidgetInstance())
+			{
+				return WidgetInstance;
+			}
+		}
+
+		return nullptr;
+	}
+
+	void FPreviewableWidgetCustomization::OnWidgetChanged(const EWidgetPreviewWidgetChangeType InChangeType)
+	{
+		PropertyUtilities->RequestForceRefresh();
+	}
+}
+
+#undef LOCTEXT_NAMESPACE
diff --git a/Engine/Plugins/Editor/UMGWidgetPreview/Source/UMGWidgetPreview/Private/Customizations/PreviewableWidgetCustomization.h b/Engine/Plugins/Editor/UMGWidgetPreview/Source/UMGWidgetPreview/Private/Customizations/PreviewableWidgetCustomization.h
new file mode 100644
--- /dev/null
+++ b/Engine/Plugins/Editor/UMGWidgetPreview/Source/UMGWidgetPreview/Private/Customizations/PreviewableWidgetCustomization.h
@@ -0,0 +1,61 @@
+// Copyright Epic Games, Inc. All Rights Reserved.
+
+#pragma once
+
+#include "Containers/Array.h"
+#include "CoreTypes.h"
+#include "Input/Reply.h"
+#include "IPropertyTypeCustomization.h"
+#include "Templates/SharedPointer.h"
+#include "UObject/WeakObjectPtr.h"
+#include "UObject/WeakObjectPtrTemplates.h"
+#include "WidgetPreview.h"
+
+class FPropertyEditorModule;
+class IDetailLayoutBuilder;
+class IPropertyUtilities;
+class SComboButton;
+class SWidget;
+class UFunction;
+class UUserWidget;
+
+namespace UE::UMGWidgetPreview::Private
+{
+	class FPreviewableWidgetCustomization
+		: public IPropertyTypeCustomization
+	{
+	public:
+		static TSharedRef<IPropertyTypeCustomization> MakeInstance();
+
+		//~ Begin IPropertyTypeCustomization
+		virtual void CustomizeHeader(TSharedRef<IPropertyHandle> PropertyHandle, FDetailWidgetRow& HeaderRow, IPropertyTypeCustomizationUtils& CustomizationUtils) override;
+		virtual void CustomizeChildren(TSharedRef<IPropertyHandle> PropertyHandle, IDetailChildrenBuilder& ChildBuilder, IPropertyTypeCustomizationUtils& CustomizationUtils) override;
+		//~ End IPropertyTypeCustomization
+
+	private:
+		bool GetCallInEditorFunctions(TArray<UFunction*>& OutCallInEditorFunctions) const;
+		void AddCallInEditorFunctions(IDetailChildrenBuilder& ChildBuilder, const TArrayView<UFunction*>& InCallInEditorFunctions);
+		
+		TArray<TWeakObjectPtr<UObject>> GetFunctionCallExecutionContext(TWeakObjectPtr<UFunction> InWeakFunction) const;
+
+		// Executes the specified method on the widget instance
+		FReply OnExecuteCallInEditorFunction(TWeakObjectPtr<UFunction> InWeakFunction);
+
+		// Checks if the function can be called - ie. that it has an object to call with.
+		bool CanExecuteCallInEditorFunction(TWeakObjectPtr<UFunction> InWeakFunction) const;
+
+		FPreviewableWidgetVariant* GetPreviewableWidgetVariant() const;
+		UUserWidget* GetWidgetInstance() const;
+
+		void OnWidgetChanged(const EWidgetPreviewWidgetChangeType InChangeType);
+
+	private:
+		TSharedPtr<IPropertyUtilities> PropertyUtilities;
+		TWeakObjectPtr<UWidgetPreview> WeakOwningPreview;
+		TSharedPtr<IPropertyHandle> OwningPreviewHandle;
+		TSharedPtr<IPropertyHandle> PreviewVariantHandle;
+
+		TSharedPtr<IPropertyHandle> ObjectPathHandle;
+		TWeakObjectPtr<UUserWidget> WeakWidgetInstance;
+	};
+}
diff --git a/Engine/Plugins/Editor/UMGWidgetPreview/Source/UMGWidgetPreview/Private/UMGWidgetPreviewModule.cpp b/Engine/Plugins/Editor/UMGWidgetPreview/Source/UMGWidgetPreview/Private/UMGWidgetPreviewModule.cpp
--- a/Engine/Plugins/Editor/UMGWidgetPreview/Source/UMGWidgetPreview/Private/UMGWidgetPreviewModule.cpp
+++ b/Engine/Plugins/Editor/UMGWidgetPreview/Source/UMGWidgetPreview/Private/UMGWidgetPreviewModule.cpp
@@ -4,6 +4,7 @@
 
 #include "Blueprint/UserWidget.h"
 #include "ContentBrowserMenuContexts.h"
+#include "Customizations/PreviewableWidgetCustomization.h"
 #include "Editor.h"
 #include "Framework/Commands/UIAction.h"
 #include "Logging/LogMacros.h"
@@ -43,12 +44,24 @@ namespace UE::UMGWidgetPreview::Private
 		// Menus need to be registered in a callback to make sure the system is ready for them.
 		UToolMenus::RegisterStartupCallback(
 			FSimpleMulticastDelegate::FDelegate::CreateRaw(this, &FUMGWidgetPreviewModule::RegisterMenus));
+
+		FPropertyEditorModule& PropertyModule = FModuleManager::GetModuleChecked<FPropertyEditorModule>("PropertyEditor");
+		{
+			PropertyModule.RegisterCustomPropertyTypeLayout(
+				FPreviewableWidgetVariant::StaticStruct()->GetFName(),
+				FOnGetPropertyTypeCustomizationInstance::CreateStatic(&FPreviewableWidgetCustomization::MakeInstance));
+		}
 	}
 
 	void FUMGWidgetPreviewModule::ShutdownModule()
 	{
 		UToolMenus::UnRegisterStartupCallback(this);
 
+		if (FPropertyEditorModule* PropertyModule = FModuleManager::GetModulePtr<FPropertyEditorModule>("PropertyEditor"))
+		{
+			PropertyModule->UnregisterCustomPropertyTypeLayout(FPreviewableWidgetVariant::StaticStruct()->GetFName());
+		}
+
 		FMessageLogModule& MessageLogModule = FModuleManager::LoadModuleChecked<FMessageLogModule>("MessageLog");
 		MessageLogModule.UnregisterLogListing(MessageLogName);
 
@@ -98,6 +111,7 @@ namespace UE::UMGWidgetPreview::Private
 													{
 														UWidgetPreview* PreviewForWidget = UWidgetPreviewEditor::CreatePreviewForWidget(
 															Cast<UUserWidget>(AsWidgetBlueprint->GeneratedClass->GetDefaultObject<UUserWidget>()));
+
 														GEditor->GetEditorSubsystem<UAssetEditorSubsystem>()->OpenEditorForAsset(PreviewForWidget);
 													}
 												}
@@ -110,7 +124,7 @@ namespace UE::UMGWidgetPreview::Private
 							Section.AddMenuEntryWithCommandList(
 								FWidgetPreviewCommands::Get().OpenEditor,
 								CommandListToBind,
-								TAttribute<FText>(),
+								LOCTEXT("WidgetContextMenuPreviewLabel", "Preview"), // Just use "Preview" here, the context means it's already a "Widget" so we can omit the prefix
 								ToolTipOverride,
 								FSlateIcon(FWidgetPreviewStyle::Get().GetStyleSetName(), "WidgetPreview.OpenEditor"));
 						}
diff --git a/Engine/Plugins/Editor/UMGWidgetPreview/Source/UMGWidgetPreview/UMGWidgetPreview.Build.cs b/Engine/Plugins/Editor/UMGWidgetPreview/Source/UMGWidgetPreview/UMGWidgetPreview.Build.cs
--- a/Engine/Plugins/Editor/UMGWidgetPreview/Source/UMGWidgetPreview/UMGWidgetPreview.Build.cs
+++ b/Engine/Plugins/Editor/UMGWidgetPreview/Source/UMGWidgetPreview/UMGWidgetPreview.Build.cs
@@ -13,12 +13,14 @@ public class UMGWidgetPreview : ModuleRules
 			{
 				"AdvancedPreviewScene",
 				"AssetDefinition",
+				"BlueprintGraph",
 				"ContentBrowser",
 				"Core",
 				"CoreUObject",
                 "DataValidation",
 				"EditorSubsystem",
 				"Engine",
+				"FieldNotification",
 				"InputCore",
 				"MessageLog",
 				"Projects",
