[UMG Widget Preview] CIS fix

diff --git a/Engine/Plugins/Editor/UMGWidgetPreview/Source/UMGWidgetPreview/Private/Customizations/PreviewableWidgetCustomization.cpp b/Engine/Plugins/Editor/UMGWidgetPreview/Source/UMGWidgetPreview/Private/Customizations/PreviewableWidgetCustomization.cpp
--- a/Engine/Plugins/Editor/UMGWidgetPreview/Source/UMGWidgetPreview/Private/Customizations/PreviewableWidgetCustomization.cpp
+++ b/Engine/Plugins/Editor/UMGWidgetPreview/Source/UMGWidgetPreview/Private/Customizations/PreviewableWidgetCustomization.cpp
@@ -3,24 +3,17 @@
 #include "PreviewableWidgetCustomization.h"
 
 #include "Blueprint/UserWidget.h"
-#include "ContentBrowserModule.h"
-#include "Customizations/MathStructCustomizations.h"
-#include "DetailLayoutBuilder.h"
 #include "Framework/Application/SlateApplication.h"
 #include "Framework/MultiBox/MultiBoxBuilder.h"
 #include "Framework/Views/ITypedTableView.h"
-#include "IContentBrowserSingleton.h"
 #include "IDetailChildrenBuilder.h"
 #include "IDetailGroup.h"
-#include "IDetailPropertyRow.h"
 #include "IPropertyUtilities.h"
-#include "K2Node_CallFunction.h"
 #include "Kismet2/BlueprintEditorUtils.h"
 #include "Modules/ModuleManager.h"
 #include "ObjectTools.h"
 #include "PropertyCustomizationHelpers.h"
 #include "PropertyEditorModule.h"
-#include "Reflection/FunctionUtils.h"
 #include "ScopedTransaction.h"
 #include "Settings/BlueprintEditorProjectSettings.h"
 #include "Styling/StyleColors.h"
@@ -305,7 +298,6 @@ namespace UE::UMGWidgetPreview::Private
 
 	FReply FPreviewableWidgetCustomization::OnExecuteCallInEditorFunction(TWeakObjectPtr<UFunction> InWeakFunction)
 	{
-		using namespace UE::Reflection;
 		TArray<TWeakObjectPtr<UObject>> WeakExecutionObjects = GetFunctionCallExecutionContext(InWeakFunction);
 		if (UFunction* Function = InWeakFunction.Get())
 		{
