[UMG Widget Preview] Added details customization for widget slot interaction

diff --git a/Engine/Plugins/Editor/UMGWidgetPreview/Source/UMGWidgetPreview/Private/Customizations/WidgetPreviewCustomization.cpp b/Engine/Plugins/Editor/UMGWidgetPreview/Source/UMGWidgetPreview/Private/Customizations/WidgetPreviewCustomization.cpp
new file mode 100644
--- /dev/null
+++ b/Engine/Plugins/Editor/UMGWidgetPreview/Source/UMGWidgetPreview/Private/Customizations/WidgetPreviewCustomization.cpp
@@ -0,0 +1,74 @@
+// Copyright Epic Games, Inc. All Rights Reserved.
+
+#include "WidgetPreviewCustomization.h"
+
+#include "Blueprint/UserWidget.h"
+#include "DetailCategoryBuilder.h"
+#include "DetailLayoutBuilder.h"
+
+#define LOCTEXT_NAMESPACE "WidgetPreviewCustomization"
+
+namespace UE::UMGWidgetPreview::Private
+{
+	const FName FWidgetPreviewCustomization::WidgetTypePropertyName = "WidgetType";
+	const FName FWidgetPreviewCustomization::SlotWidgetTypesPropertyName = "SlotWidgetTypes";
+	const FName FWidgetPreviewCustomization::WidgetInstancePropertyName = "WidgetInstance";
+
+	TSharedRef<IDetailCustomization> FWidgetPreviewCustomization::MakeInstance()
+	{
+		TSharedRef<FWidgetPreviewCustomization> Customization = MakeShared<FWidgetPreviewCustomization>();
+		return Customization;
+	}
+
+	void FWidgetPreviewCustomization::CustomizeDetails(IDetailLayoutBuilder& DetailBuilder)
+	{
+		TArray<TWeakObjectPtr<UWidgetPreview>> CustomizedWidgetPreviews = DetailBuilder.GetObjectsOfTypeBeingCustomized<UWidgetPreview>();
+
+		TSharedRef<IPropertyHandle> WidgetTypePropertyHandle = DetailBuilder.GetProperty(WidgetTypePropertyName);
+		TSharedRef<IPropertyHandle> SlotWidgetTypesPropertyHandle = DetailBuilder.GetProperty(SlotWidgetTypesPropertyName);
+		TSharedRef<IPropertyHandle> WidgetInstancePropertyHandle = DetailBuilder.GetProperty(WidgetInstancePropertyName);
+
+		static const FName WidgetCategoryName = "Widget";
+
+		for (TWeakObjectPtr<UWidgetPreview>& WeakWidgetPreview : CustomizedWidgetPreviews)
+		{
+			if (UWidgetPreview* WidgetPreview = WeakWidgetPreview.Get())
+			{
+				IDetailCategoryBuilder& WidgetCategory = DetailBuilder.EditCategory(WidgetCategoryName);
+				WidgetCategory.AddProperty(WidgetTypePropertyHandle);
+
+				// Widget Instance
+				{
+					FAddPropertyParams AddPropertyParams;
+
+					WidgetCategory.AddExternalObjects(
+						{ WidgetPreview->GetWidgetInstance() },
+						EPropertyLocation::Default,
+						AddPropertyParams)
+					->DisplayName(LOCTEXT("WidgetInstanceHeader", "Widget Properties"))
+					.ShouldAutoExpand(true);
+				}
+
+				// Slot Widgets
+				if (SlotWidgetTypesPropertyHandle->IsValidHandle())
+				{
+					TAttribute<bool> SlotWidgetTypesEditCondition = TAttribute<bool>::CreateLambda(
+						[WeakWidgetPreview]()
+						{
+							if (UWidgetPreview* StrongWidgetPreview = WeakWidgetPreview.Get())
+							{
+								return !StrongWidgetPreview->GetWidgetSlotNames().IsEmpty();
+							}
+
+							return false;
+						});
+
+					WidgetCategory.AddProperty(SlotWidgetTypesPropertyHandle)
+					.EditCondition(SlotWidgetTypesEditCondition, { });
+				}
+			}
+		}
+	}
+}
+
+#undef LOCTEXT_NAMESPACE
diff --git a/Engine/Plugins/Editor/UMGWidgetPreview/Source/UMGWidgetPreview/Private/Customizations/WidgetPreviewCustomization.h b/Engine/Plugins/Editor/UMGWidgetPreview/Source/UMGWidgetPreview/Private/Customizations/WidgetPreviewCustomization.h
new file mode 100644
--- /dev/null
+++ b/Engine/Plugins/Editor/UMGWidgetPreview/Source/UMGWidgetPreview/Private/Customizations/WidgetPreviewCustomization.h
@@ -0,0 +1,40 @@
+// Copyright Epic Games, Inc. All Rights Reserved.
+
+#pragma once
+
+#include "Containers/Array.h"
+#include "CoreTypes.h"
+#include "IDetailCustomization.h"
+#include "Input/Reply.h"
+#include "IPropertyTypeCustomization.h"
+#include "Templates/SharedPointer.h"
+#include "UObject/WeakObjectPtr.h"
+#include "UObject/WeakObjectPtrTemplates.h"
+#include "WidgetPreview.h"
+
+class FPropertyEditorModule;
+class IDetailLayoutBuilder;
+class IPropertyUtilities;
+class SComboButton;
+class SWidget;
+class UFunction;
+class UUserWidget;
+
+namespace UE::UMGWidgetPreview::Private
+{
+	class FWidgetPreviewCustomization
+		: public IDetailCustomization
+	{
+	public:
+		static TSharedRef<IDetailCustomization> MakeInstance();
+
+		//~ Begin IDetailCustomization
+		virtual void CustomizeDetails(IDetailLayoutBuilder& DetailBuilder) override;
+		//~ End IDetailCustomization
+
+	private:
+		static const FName WidgetTypePropertyName;
+        static const FName SlotWidgetTypesPropertyName;
+        static const FName WidgetInstancePropertyName;
+	};
+}
diff --git a/Engine/Plugins/Editor/UMGWidgetPreview/Source/UMGWidgetPreview/Private/UMGWidgetPreviewModule.cpp b/Engine/Plugins/Editor/UMGWidgetPreview/Source/UMGWidgetPreview/Private/UMGWidgetPreviewModule.cpp
--- a/Engine/Plugins/Editor/UMGWidgetPreview/Source/UMGWidgetPreview/Private/UMGWidgetPreviewModule.cpp
+++ b/Engine/Plugins/Editor/UMGWidgetPreview/Source/UMGWidgetPreview/Private/UMGWidgetPreviewModule.cpp
@@ -5,6 +5,7 @@
 #include "Blueprint/UserWidget.h"
 #include "ContentBrowserMenuContexts.h"
 #include "Customizations/PreviewableWidgetCustomization.h"
+#include "Customizations/WidgetPreviewCustomization.h"
 #include "Editor.h"
 #include "Framework/Commands/UIAction.h"
 #include "Logging/LogMacros.h"
@@ -50,6 +51,10 @@ namespace UE::UMGWidgetPreview::Private
 			PropertyModule.RegisterCustomPropertyTypeLayout(
 				FPreviewableWidgetVariant::StaticStruct()->GetFName(),
 				FOnGetPropertyTypeCustomizationInstance::CreateStatic(&FPreviewableWidgetCustomization::MakeInstance));
+
+			PropertyModule.RegisterCustomClassLayout(
+				UWidgetPreview::StaticClass()->GetFName(),
+				FOnGetDetailCustomizationInstance::CreateStatic(&FWidgetPreviewCustomization::MakeInstance));
 		}
 	}
 
@@ -59,6 +64,7 @@ namespace UE::UMGWidgetPreview::Private
 
 		if (FPropertyEditorModule* PropertyModule = FModuleManager::GetModulePtr<FPropertyEditorModule>("PropertyEditor"))
 		{
+			PropertyModule->UnregisterCustomClassLayout(UWidgetPreview::StaticClass()->GetFName());
 			PropertyModule->UnregisterCustomPropertyTypeLayout(FPreviewableWidgetVariant::StaticStruct()->GetFName());
 		}
 
