[UMG Widget Preview] Fixed Toolkit Save workflow

diff --git a/Engine/Plugins/Editor/UMGWidgetPreview/Source/UMGWidgetPreview/Private/WidgetPreviewToolkit.cpp b/Engine/Plugins/Editor/UMGWidgetPreview/Source/UMGWidgetPreview/Private/WidgetPreviewToolkit.cpp
--- a/Engine/Plugins/Editor/UMGWidgetPreview/Source/UMGWidgetPreview/Private/WidgetPreviewToolkit.cpp
+++ b/Engine/Plugins/Editor/UMGWidgetPreview/Source/UMGWidgetPreview/Private/WidgetPreviewToolkit.cpp
@@ -414,13 +414,22 @@ namespace UE::UMGWidgetPreview::Private
 			return;
 		}
 
+		// Check for Transient outer, and if found use SaveAs instead
+		for (UObject* Object : ObjectsToSave)
+		{
+			UPackage* Package = Object->GetPackage();
+			if (!Package || Package == GetTransientPackage())
+			{
+				// Redirect to SaveAs
+				SaveAssetAs_Execute();
+				return;
+			}
+		}
+
 		TArray<UObject*> SavedObjects;
 		SavedObjects.Reserve(ObjectsToSave.Num());
 
 		TArray<UPackage*> PackagesToSave;
-		bool bHasNewlyCreatedPackage = false;
-
-		TArray<UPackage*> PackagesToSaveAs;
 
 		for (UObject* Object : ObjectsToSave)
 		{
@@ -431,11 +440,13 @@ namespace UE::UMGWidgetPreview::Private
 			}
 			else
 			{
+				PackagesToSave.Add(Object->GetOutermost());
 				SavedObjects.Add(Object);
 			}
 		}
 
-		FEditorFileUtils::PromptForCheckoutAndSave(PackagesToSave, bCheckDirtyOnAssetSave, /*bPromptToSave=*/ bHasNewlyCreatedPackage);
+		constexpr bool bPromptToSave = false;
+		FEditorFileUtils::PromptForCheckoutAndSave(PackagesToSave, bCheckDirtyOnAssetSave, bPromptToSave);
 
 		OnAssetsSaved(SavedObjects);
 	}
