[UMG Widget Preview] Added public API to allow extensibility (primarily for MVVM)

diff --git a/Engine/Plugins/Editor/UMGWidgetPreview/Source/UMGWidgetPreview/Private/Customizations/PreviewableWidgetCustomization.cpp b/Engine/Plugins/Editor/UMGWidgetPreview/Source/UMGWidgetPreview/Private/Customizations/PreviewableWidgetCustomization.cpp
--- a/Engine/Plugins/Editor/UMGWidgetPreview/Source/UMGWidgetPreview/Private/Customizations/PreviewableWidgetCustomization.cpp
+++ b/Engine/Plugins/Editor/UMGWidgetPreview/Source/UMGWidgetPreview/Private/Customizations/PreviewableWidgetCustomization.cpp
@@ -42,7 +42,6 @@ namespace UE::UMGWidgetPreview::Private
 		IPropertyTypeCustomizationUtils& CustomizationUtils)
 	{
 		PropertyUtilities = CustomizationUtils.GetPropertyUtilities();
-		OwningPreviewHandle = PropertyHandle->GetParentHandle();
 		PreviewVariantHandle = PropertyHandle;
 		ObjectPathHandle = PreviewVariantHandle->GetChildHandle(GET_MEMBER_NAME_CHECKED(FPreviewableWidgetVariant, ObjectPath));
 
@@ -93,7 +92,7 @@ namespace UE::UMGWidgetPreview::Private
 		auto CanDisplayAndCallFunction = [](const UFunction* TestFunction)
 		{
 			bool bCanCall = TestFunction->GetBoolMetaData(FBlueprintMetadata::MD_CallInEditor)
-						&& TestFunction->HasAnyFunctionFlags(FUNC_Private)
+						&& TestFunction->HasAnyFunctionFlags(FUNC_Public)
 						&& TestFunction->ParmsSize == 0; // Params not supported
 
 			return bCanCall;
diff --git a/Engine/Plugins/Editor/UMGWidgetPreview/Source/UMGWidgetPreview/Private/Customizations/PreviewableWidgetCustomization.h b/Engine/Plugins/Editor/UMGWidgetPreview/Source/UMGWidgetPreview/Private/Customizations/PreviewableWidgetCustomization.h
--- a/Engine/Plugins/Editor/UMGWidgetPreview/Source/UMGWidgetPreview/Private/Customizations/PreviewableWidgetCustomization.h
+++ b/Engine/Plugins/Editor/UMGWidgetPreview/Source/UMGWidgetPreview/Private/Customizations/PreviewableWidgetCustomization.h
@@ -52,9 +52,7 @@ namespace UE::UMGWidgetPreview::Private
 	private:
 		TSharedPtr<IPropertyUtilities> PropertyUtilities;
 		TWeakObjectPtr<UWidgetPreview> WeakOwningPreview;
-		TSharedPtr<IPropertyHandle> OwningPreviewHandle;
 		TSharedPtr<IPropertyHandle> PreviewVariantHandle;
-
 		TSharedPtr<IPropertyHandle> ObjectPathHandle;
 		TWeakObjectPtr<UUserWidget> WeakWidgetInstance;
 	};
diff --git a/Engine/Plugins/Editor/UMGWidgetPreview/Source/UMGWidgetPreview/Private/UMGWidgetPreviewModule.cpp b/Engine/Plugins/Editor/UMGWidgetPreview/Source/UMGWidgetPreview/Private/UMGWidgetPreviewModule.cpp
--- a/Engine/Plugins/Editor/UMGWidgetPreview/Source/UMGWidgetPreview/Private/UMGWidgetPreviewModule.cpp
+++ b/Engine/Plugins/Editor/UMGWidgetPreview/Source/UMGWidgetPreview/Private/UMGWidgetPreviewModule.cpp
@@ -1,4 +1,4 @@
-﻿// Copyright Epic Games, Inc. All Rights Reserved.
+// Copyright Epic Games, Inc. All Rights Reserved.
 
 #include "UMGWidgetPreviewModule.h"
 
@@ -68,12 +68,19 @@ namespace UE::UMGWidgetPreview::Private
 			PropertyModule->UnregisterCustomPropertyTypeLayout(FPreviewableWidgetVariant::StaticStruct()->GetFName());
 		}
 
-		FMessageLogModule& MessageLogModule = FModuleManager::LoadModuleChecked<FMessageLogModule>("MessageLog");
-		MessageLogModule.UnregisterLogListing(MessageLogName);
+		if (FMessageLogModule* MessageLogModule = FModuleManager::GetModulePtr<FMessageLogModule>("MessageLog"))
+		{
+			MessageLogModule->UnregisterLogListing(MessageLogName);
+		}
 
 		FWidgetPreviewCommands::Unregister();
 	}
 
+	IUMGWidgetPreviewModule::FOnRegisterTabs& FUMGWidgetPreviewModule::OnRegisterTabsForEditor()
+	{
+		return RegisterTabsForEditorDelegate;
+	}
+
 	void FUMGWidgetPreviewModule::RegisterMenus()
 	{
 		// Allows cleanup when module unloads.
diff --git a/Engine/Plugins/Editor/UMGWidgetPreview/Source/UMGWidgetPreview/Private/UMGWidgetPreviewModule.h b/Engine/Plugins/Editor/UMGWidgetPreview/Source/UMGWidgetPreview/Private/UMGWidgetPreviewModule.h
--- a/Engine/Plugins/Editor/UMGWidgetPreview/Source/UMGWidgetPreview/Private/UMGWidgetPreviewModule.h
+++ b/Engine/Plugins/Editor/UMGWidgetPreview/Source/UMGWidgetPreview/Private/UMGWidgetPreviewModule.h
@@ -2,12 +2,12 @@
 
 #pragma once
 
-#include "Modules/ModuleInterface.h"
+#include "UMGWidgetPreview/Public/IUMGWidgetPreviewModule.h"
 
 namespace UE::UMGWidgetPreview::Private
 {
 	class FUMGWidgetPreviewModule
-		: public IModuleInterface
+		: public IUMGWidgetPreviewModule
 	{
 	public:
 		//~ Begin IModuleInterface
@@ -15,7 +15,14 @@ namespace UE::UMGWidgetPreview::Private
 		virtual void ShutdownModule() override;
 		//~ End IModuleInterface
 
+		//~ Begin IUMGWidgetPreviewModule
+		virtual FOnRegisterTabs& OnRegisterTabsForEditor() override;
+		//~ End IUMGWidgetPreviewModule
+
 	private:
 		void RegisterMenus();
+
+	private:
+		FOnRegisterTabs RegisterTabsForEditorDelegate;
 	};
 }
diff --git a/Engine/Plugins/Editor/UMGWidgetPreview/Source/UMGWidgetPreview/Private/WidgetPreview.cpp b/Engine/Plugins/Editor/UMGWidgetPreview/Source/UMGWidgetPreview/Private/WidgetPreview.cpp
--- a/Engine/Plugins/Editor/UMGWidgetPreview/Source/UMGWidgetPreview/Private/WidgetPreview.cpp
+++ b/Engine/Plugins/Editor/UMGWidgetPreview/Source/UMGWidgetPreview/Private/WidgetPreview.cpp
@@ -3,9 +3,13 @@
 #include "WidgetPreview.h"
 
 #include "Blueprint/UserWidget.h"
+#include "Blueprint/WidgetTree.h"
+#include "Engine/Engine.h"
+#include "UObject/Package.h"
 #include "WidgetBlueprint.h"
-#include "WidgetPreviewTypesPrivate.h"
 #include "WidgetPreviewLog.h"
+#include "WidgetPreviewTypesPrivate.h"
+#include "Widgets/SWidget.h"
 
 FPreviewableWidgetVariant::FPreviewableWidgetVariant(const TSubclassOf<UUserWidget>& InWidgetType)
 	: ObjectPath(InWidgetType)
@@ -21,7 +25,9 @@ FPreviewableWidgetVariant::FPreviewableWidgetVariant(const UWidgetPreview* InWid
 
 void FPreviewableWidgetVariant::UpdateCachedWidget()
 {
-	CachedWidgetCDO.Reset();
+	UE::UMGWidgetPreview::Private::FWidgetTypeTuple PreviewUserWidgetTuple;
+
+	const UUserWidget* PreviewUserWidgetCDO = nullptr;
 	CachedWidgetPreview.Reset();
 
 	if (UObject* ResolvedObject = ObjectPath.TryLoad())
@@ -29,20 +35,23 @@ void FPreviewableWidgetVariant::UpdateCachedWidget()
     	if (UWidgetPreview* WidgetPreview = Cast<UWidgetPreview>(ResolvedObject))
     	{
     		CachedWidgetPreview = WidgetPreview;
-    		CachedWidgetCDO = WidgetPreview->GetWidgetCDO();
+    		PreviewUserWidgetCDO = WidgetPreview->GetWidgetCDO();
+    		PreviewUserWidgetTuple.Set(WidgetPreview->GetWidgetCDO());
     	}
 		else if (const UWidgetBlueprint* AsBlueprint = Cast<UWidgetBlueprint>(ResolvedObject))
 		{
-			CachedWidgetCDO = Cast<UUserWidget>(AsBlueprint->GeneratedClass->GetDefaultObject<UUserWidget>());
+			PreviewUserWidgetCDO = AsBlueprint->GeneratedClass->GetDefaultObject<UUserWidget>();
 		}
 		else if (const UClass* AsClass = Cast<UClass>(ResolvedObject))
 		{
-			if (UUserWidget* UserWidget = AsClass->GetDefaultObject<UUserWidget>())
+			if (const UUserWidget* UserWidgetCDO = AsClass->GetDefaultObject<UUserWidget>())
 			{
-				CachedWidgetCDO = UserWidget;
+				PreviewUserWidgetCDO = UserWidgetCDO;
 			}
 		}
     }
+
+	CachedWidgetCDO = PreviewUserWidgetCDO;
 }
 
 const UUserWidget* FPreviewableWidgetVariant::AsUserWidgetCDO() const
@@ -72,7 +81,7 @@ const UUserWidget* FPreviewableWidgetVariant::AsUserWidgetCDO() const
 	return nullptr;
 }
 
-UWidgetPreview* FPreviewableWidgetVariant::AsWidgetPreview() const
+const UWidgetPreview* FPreviewableWidgetVariant::AsWidgetPreview() const
 {
 	if (ObjectPath.IsNull())
 	{
@@ -105,13 +114,18 @@ UUserWidget* UWidgetPreview::GetOrCreateWidgetInstance(UWorld* InWorld, const bo
 {
 	if (bInForceRecreate)
 	{
-		WidgetInstance = nullptr;
+		ClearWidgetInstance();
 	}
 	else if (WidgetInstance)
 	{
 		return WidgetInstance;
 	}
 
+	return CreateWidgetInstance(InWorld);
+}
+
+UUserWidget* UWidgetPreview::CreateWidgetInstance(UWorld* InWorld)
+{
 	if (!ensure(InWorld))
 	{
 		return nullptr;
@@ -128,13 +142,13 @@ UUserWidget* UWidgetPreview::GetOrCreateWidgetInstance(UWorld* InWorld, const bo
 	{
 		auto MakeWidget = [InWorld](UClass* InClass) -> UUserWidget*
 		{
+			ensureAlways(!InClass->GetName().StartsWith(TEXT("REINST_")));
+
 			UUserWidget* NewWidget = NewObject<UUserWidget>(InWorld, InClass);
 			NewWidget->ClearFlags(RF_Transactional);
 			return NewWidget;
 		};
 
-		// @todo: always wrap with a container UWidget to account for warning here: void UWidget::RemoveFromParent()
-
 		WidgetInstance = MakeWidget(Widget->GetClass());
 
 		if (!WidgetType.ObjectPath.IsNull() && !SlotWidgetTypes.IsEmpty())
@@ -157,6 +171,8 @@ UUserWidget* UWidgetPreview::GetOrCreateWidgetInstance(UWorld* InWorld, const bo
 			WidgetInstance->SetPlayerContext(LocalPlayer);
 		}
 
+		SlateWidgetInstance = WidgetInstance->TakeWidget();
+
 		OnWidgetChanged().Broadcast(EWidgetPreviewWidgetChangeType::Reinstanced);
 
 		return WidgetInstance;
@@ -170,6 +186,21 @@ UUserWidget* UWidgetPreview::GetWidgetInstance() const
 	return WidgetInstance;
 }
 
+TSharedPtr<SWidget> UWidgetPreview::GetSlateWidgetInstance() const
+{
+	if (SlateWidgetInstance.IsValid())
+	{
+		return SlateWidgetInstance;
+	}
+
+	if (UUserWidget* Instance = GetWidgetInstance())
+	{
+		return Instance->TakeWidget();
+	}
+
+	return nullptr;
+}
+
 const UUserWidget* UWidgetPreview::GetWidgetCDO() const
 {
 	// If the LayoutWidget is in use, return it (as the root widget).
@@ -299,7 +330,7 @@ void UWidgetPreview::SetWidgetType(const FPreviewableWidgetVariant& InWidget)
 	if (WidgetType != InWidget)
 	{
 		WidgetType = InWidget;
-		WidgetInstance = nullptr;
+		ClearWidgetInstance();
 		UpdateWidgets();
 
 		OnWidgetChanged().Broadcast(EWidgetPreviewWidgetChangeType::Assignment);
@@ -316,7 +347,7 @@ void UWidgetPreview::SetSlotWidgetTypes(const TMap<FName, FPreviewableWidgetVari
 	if (!SlotWidgetTypes.OrderIndependentCompareEqual(InWidgets))
 	{
 		SlotWidgetTypes = InWidgets;
-		WidgetInstance = nullptr;
+		ClearWidgetInstance();
 		UpdateWidgets();
 
 		OnWidgetChanged().Broadcast(EWidgetPreviewWidgetChangeType::Assignment);
@@ -339,24 +370,39 @@ void UWidgetPreview::PostEditChangeProperty(FPropertyChangedEvent& PropertyChang
 		|| PropertyName == GET_MEMBER_NAME_CHECKED(UWidgetPreview, SlotWidgetTypes)
 		|| PropertyName.IsNone()) // None can be an Undo operation
 	{
-		WidgetInstance = nullptr;
+		ClearWidgetInstance();
 		UpdateWidgets();
 		OnWidgetChanged().Broadcast(EWidgetPreviewWidgetChangeType::Assignment);
 	}
 }
 
-void UWidgetPreview::PostEditChangeChainProperty(FPropertyChangedChainEvent& PropertyChangedEvent)
-{
-	UObject::PostEditChangeChainProperty(PropertyChangedEvent);
-}
-
+// Due to compilation, etc.
 void UWidgetPreview::OnWidgetBlueprintChanged(UBlueprint* InBlueprint)
 {
-	WidgetInstance = nullptr;
+	ClearWidgetInstance();
 	UpdateWidgets();
 	OnWidgetChanged().Broadcast(EWidgetPreviewWidgetChangeType::Structure);
 }
 
+void UWidgetPreview::ClearWidgetInstance()
+{
+	if (WidgetInstance)
+	{
+		if (SlateWidgetInstance.IsValid())
+		{
+			SlateWidgetInstance.Reset();
+		}
+
+		OnWidgetChanged().Broadcast(EWidgetPreviewWidgetChangeType::Destroyed);
+
+		WidgetInstance->OnNativeDestruct.RemoveAll(this);
+		WidgetInstance->MarkAsGarbage();
+		WidgetInstance = nullptr;
+
+		CollectGarbage(GARBAGE_COLLECTION_KEEPFLAGS);
+	}
+}
+
 void UWidgetPreview::UpdateWidgets()
 {
 	auto AddOnChanged = [this](const UUserWidget* InUserWidgetCDO)
diff --git a/Engine/Plugins/Editor/UMGWidgetPreview/Source/UMGWidgetPreview/Private/WidgetPreview.h b/Engine/Plugins/Editor/UMGWidgetPreview/Source/UMGWidgetPreview/Private/WidgetPreview.h
deleted file mode 100644
diff --git a/Engine/Plugins/Editor/UMGWidgetPreview/Source/UMGWidgetPreview/Private/WidgetPreviewToolkit.cpp b/Engine/Plugins/Editor/UMGWidgetPreview/Source/UMGWidgetPreview/Private/WidgetPreviewToolkit.cpp
--- a/Engine/Plugins/Editor/UMGWidgetPreview/Source/UMGWidgetPreview/Private/WidgetPreviewToolkit.cpp
+++ b/Engine/Plugins/Editor/UMGWidgetPreview/Source/UMGWidgetPreview/Private/WidgetPreviewToolkit.cpp
@@ -16,7 +16,7 @@
 #include "ToolMenuMisc.h"
 #include "ToolMenus.h"
 #include "ToolMenuSection.h"
-#include "WidgetBlueprintEditor.h"
+#include "UMGWidgetPreview/Public/IUMGWidgetPreviewModule.h"
 #include "WidgetPreview.h"
 #include "WidgetPreviewCommands.h"
 #include "WidgetPreviewEditor.h"
@@ -137,6 +137,11 @@ namespace UE::UMGWidgetPreview::Private
 		return PreviewScene.ToSharedRef();
 	}
 
+	FWidgetPreviewToolkitStateBase::FWidgetPreviewToolkitStateBase(const FName& Id)
+		: Id(Id)
+	{
+	}
+
 	FName FWidgetPreviewToolkitStateBase::GetId() const
 	{
 		return Id;
@@ -159,10 +164,12 @@ namespace UE::UMGWidgetPreview::Private
 
 	void FWidgetPreviewToolkitStateBase::OnEnter(const FWidgetPreviewToolkitStateBase* InFromState)
 	{
+		// Default, empty Implementation
 	}
 
 	void FWidgetPreviewToolkitStateBase::OnExit(const FWidgetPreviewToolkitStateBase* InToState)
 	{
+		// Default, empty Implementation
 	}
 
 	FWidgetPreviewToolkitPausedState::FWidgetPreviewToolkitPausedState(): FWidgetPreviewToolkitStateBase(TEXT("Paused"))
@@ -178,16 +185,6 @@ namespace UE::UMGWidgetPreview::Private
 		StatusMessage = FTokenizedMessage::Create(EMessageSeverity::Info, LOCTEXT("WidgetPreviewToolkitBackgroundState_Message", "The widget preview is paused while the window is in the background. Re-focus to unpause."));
 	}
 
-	void FWidgetPreviewToolkitBackgroundState::OnEnter(const FWidgetPreviewToolkitStateBase* InFromState)
-	{
-		FWidgetPreviewToolkitPausedState::OnEnter(InFromState);
-	}
-
-	void FWidgetPreviewToolkitBackgroundState::OnExit(const FWidgetPreviewToolkitStateBase* InToState)
-	{
-		FWidgetPreviewToolkitPausedState::OnExit(InToState);
-	}
-
 	FWidgetPreviewToolkitUnsupportedWidgetState::FWidgetPreviewToolkitUnsupportedWidgetState()
 	{
 		Id = TEXT("UnsupportedWidget");
@@ -278,8 +275,14 @@ namespace UE::UMGWidgetPreview::Private
 
 	FWidgetPreviewToolkit::~FWidgetPreviewToolkit()
 	{
+		if (GEditor)
+		{
+			GEditor->OnBlueprintPreCompile().Remove(OnBlueprintPrecompileHandle);
+		}
+
 		if (Preview)
 		{
+			Preview->ClearWidgetInstance();
 			Preview->OnWidgetChanged().Remove(OnWidgetChangedHandle);
 		}
 
@@ -312,6 +315,11 @@ namespace UE::UMGWidgetPreview::Private
 			nullptr,
 			TArray<FAdvancedPreviewSceneModule::FDetailCustomizationInfo>(),  TArray<FAdvancedPreviewSceneModule::FPropertyTypeCustomizationInfo>(),
 			Delegates);
+
+		if (Preview)
+		{
+			Preview->GetOrCreateWidgetInstance(GetPreviewWorld(), true);
+		}
 	}
 
 	void FWidgetPreviewToolkit::RegisterToolbar()
@@ -338,7 +346,8 @@ namespace UE::UMGWidgetPreview::Private
 				InsertAfterAssetSection);
 
 			PreviewSection.AddEntry(
-				FToolMenuEntry::InitToolBarButton(Commands.ResetPreview,
+				FToolMenuEntry::InitToolBarButton(
+					Commands.ResetPreview,
 					FText::GetEmpty(),
 					{},
 					FSlateIcon(FWidgetPreviewStyle::Get().GetStyleSetName(), "WidgetPreview.Reset")));
@@ -367,6 +376,11 @@ namespace UE::UMGWidgetPreview::Private
 			.SetDisplayName(LOCTEXT("MessageLogTab", "Message Log"))
 			.SetGroup(AssetEditorTabsCategoryRef)
 			.SetIcon(FSlateIcon(FAppStyle::GetAppStyleSetName(), "MessageLog.TabIcon"));
+
+		const TSharedRef<IWidgetPreviewToolkit> ToolkitRef = SharedThis(this);
+
+		IUMGWidgetPreviewModule& UMGWidgetPreviewModule = FModuleManager::LoadModuleChecked<IUMGWidgetPreviewModule>("UMGWidgetPreview");
+		UMGWidgetPreviewModule.OnRegisterTabsForEditor().Broadcast(ToolkitRef, InTabManager);
 	}
 
 	void FWidgetPreviewToolkit::UnregisterTabSpawners(const TSharedRef<FTabManager>& InTabManager)
@@ -379,6 +393,11 @@ namespace UE::UMGWidgetPreview::Private
 
 	void FWidgetPreviewToolkit::PostInitAssetEditor()
 	{
+		if (GEditor)
+		{
+			OnBlueprintPrecompileHandle = GEditor->OnBlueprintPreCompile().AddRaw(this, &FWidgetPreviewToolkit::OnBlueprintPrecompile);
+		}
+
 		OnWidgetChangedHandle = Preview->OnWidgetChanged().AddSP(this, &FWidgetPreviewToolkit::OnWidgetChanged);
 
 		if (FSlateApplication::IsInitialized())
@@ -593,6 +612,38 @@ namespace UE::UMGWidgetPreview::Private
 		return TEXT("FWidgetPreviewToolkit");
 	}
 
+	TSharedPtr<FLayoutExtender> FWidgetPreviewToolkit::GetLayoutExtender() const
+	{
+		return LayoutExtender;
+	}
+
+	IWidgetPreviewToolkit::FOnSelectedObjectsChanged& FWidgetPreviewToolkit::OnSelectedObjectsChanged()
+	{
+		return SelectedObjectsChangedDelegate;
+	}
+
+	TConstArrayView<TWeakObjectPtr<UObject>> FWidgetPreviewToolkit::GetSelectedObjects() const
+	{
+		return SelectedObjects;
+	}
+
+	void FWidgetPreviewToolkit::SetSelectedObjects(const TArray<TWeakObjectPtr<UObject>>& InObjects)
+	{
+		if (InObjects.IsEmpty())
+		{
+			SelectedObjects = { Preview };
+		}
+		else
+		{
+			SelectedObjects = InObjects;
+		}
+
+		if (SelectedObjectsChangedDelegate.IsBound())
+		{
+			SelectedObjectsChangedDelegate.Broadcast(SelectedObjects);
+		}
+	}
+
 	UWidgetPreview* FWidgetPreviewToolkit::GetPreview() const
 	{
 		return Preview;
@@ -603,30 +654,30 @@ namespace UE::UMGWidgetPreview::Private
 		return CurrentState;
 	}
 
-	IWidgetPreviewToolkit::FOnStateChanged& FWidgetPreviewToolkit::OnStateChanged()
+	FWidgetPreviewToolkit::FOnStateChanged& FWidgetPreviewToolkit::OnStateChanged()
 	{
 		return OnStateChangedDelegate;
 	}
 
-	TSharedPtr<FWidgetPreviewScene> FWidgetPreviewToolkit::GetPreviewScene()
+	UWorld* FWidgetPreviewToolkit::GetPreviewWorld()
 	{
-		if (!PreviewScene.IsValid())
+		if (TSharedPtr<FWidgetPreviewScene> PreviewScenePtr = GetPreviewScene())
 		{
-			const TSharedRef<FWidgetPreviewToolkit>& Self = SharedThis(this);
-			PreviewScene = MakeShared<FWidgetPreviewScene>(Self);
+			return PreviewScenePtr->GetWorld();
 		}
 
-		return PreviewScene;
+		return nullptr;
 	}
 
-	UWorld* FWidgetPreviewToolkit::GetPreviewWorld()
+	TSharedPtr<FWidgetPreviewScene> FWidgetPreviewToolkit::GetPreviewScene()
 	{
-		if (TSharedPtr<FWidgetPreviewScene> PreviewScenePtr = GetPreviewScene())
+		if (!PreviewScene.IsValid())
 		{
-			return PreviewScenePtr->GetWorld();
+			const TSharedRef<FWidgetPreviewToolkit>& Self = SharedThis(this);
+			PreviewScene = MakeShared<FWidgetPreviewScene>(Self);
 		}
 
-		return nullptr;
+		return PreviewScene;
 	}
 
 	bool FWidgetPreviewToolkit::ShouldUpdate() const
@@ -639,6 +690,21 @@ namespace UE::UMGWidgetPreview::Private
 		return bIsFocused;
 	}
 
+	void FWidgetPreviewToolkit::OnBlueprintPrecompile(UBlueprint* InBlueprint)
+	{
+		if (Preview)
+		{
+			if (const UUserWidget* WidgetCDO = Preview->GetWidgetCDO())
+			{
+				if (InBlueprint && InBlueprint->GeneratedClass
+					&& WidgetCDO->IsA(InBlueprint->GeneratedClass))
+				{
+					Preview->ClearWidgetInstance();
+				}
+			}
+		}
+	}
+
 	void FWidgetPreviewToolkit::OnWidgetChanged(const EWidgetPreviewWidgetChangeType InChangeType)
 	{
 		ResolveState();
@@ -771,15 +837,19 @@ namespace UE::UMGWidgetPreview::Private
 	TSharedRef<SDockTab> FWidgetPreviewToolkit::SpawnTab_Details(const FSpawnTabArgs& Args)
 	{
 		check(Args.GetTabId() == DetailsTabID);
+
+		const TSharedRef<FWidgetPreviewToolkit>& Self = SharedThis(this);
+
 		return SNew(SDockTab)
 		[
-			SNew(SWidgetPreviewDetails, Preview)
+			SNew(SWidgetPreviewDetails, Self)
 		];
 	}
 
 	TSharedRef<SDockTab> FWidgetPreviewToolkit::SpawnTab_PreviewSceneSettings(const FSpawnTabArgs& Args)
 	{
 		check(Args.GetTabId().TabType == PreviewSceneSettingsTabID);
+
 		return SNew(SDockTab)
 		.Label( LOCTEXT("StaticMeshPreviewScene_TabTitle", "Preview Scene Settings") )
 		[
diff --git a/Engine/Plugins/Editor/UMGWidgetPreview/Source/UMGWidgetPreview/Private/WidgetPreviewToolkit.h b/Engine/Plugins/Editor/UMGWidgetPreview/Source/UMGWidgetPreview/Private/WidgetPreviewToolkit.h
--- a/Engine/Plugins/Editor/UMGWidgetPreview/Source/UMGWidgetPreview/Private/WidgetPreviewToolkit.h
+++ b/Engine/Plugins/Editor/UMGWidgetPreview/Source/UMGWidgetPreview/Private/WidgetPreviewToolkit.h
@@ -5,13 +5,13 @@
 #include "AdvancedPreviewScene.h"
 #include "AdvancedPreviewSceneModule.h"
 #include "Framework/Docking/TabManager.h"
+#include "IWidgetPreviewToolkit.h"
 #include "Logging/TokenizedMessage.h"
 #include "Misc/DataValidation/Fixer.h"
 #include "Templates/SharedPointer.h"
 #include "Toolkits/BaseToolkit.h"
 #include "Tools/BaseAssetToolkit.h"
 #include "WidgetPreview.h"
-#include "WorkflowOrientedApp/WorkflowTabManager.h"
 
 class UWidgetPreviewEditor;
 class FWidgetBlueprintEditor;
@@ -58,10 +58,7 @@ namespace UE::UMGWidgetPreview::Private
 
 	struct FWidgetPreviewToolkitStateBase
 	{
-		explicit FWidgetPreviewToolkitStateBase(const FName& Id)
-			: Id(Id)
-		{
-		}
+		explicit FWidgetPreviewToolkitStateBase(const FName& Id);
 
 		virtual ~FWidgetPreviewToolkitStateBase() = default;
 
@@ -88,9 +85,6 @@ namespace UE::UMGWidgetPreview::Private
 	struct FWidgetPreviewToolkitBackgroundState : FWidgetPreviewToolkitPausedState
 	{
 		FWidgetPreviewToolkitBackgroundState();
-
-		virtual void OnEnter(const FWidgetPreviewToolkitStateBase* InFromState) override;
-		virtual void OnExit(const FWidgetPreviewToolkitStateBase* InToState) override;
 	};
 
 	struct FWidgetPreviewToolkitUnsupportedWidgetState : FWidgetPreviewToolkitPausedState
@@ -111,19 +105,6 @@ namespace UE::UMGWidgetPreview::Private
 		FWidgetPreviewToolkitRunningState();
 	};
 
-	class IWidgetPreviewToolkit
-	{
-	public:
-		virtual ~IWidgetPreviewToolkit() = default;
-
-		using FOnStateChanged = TMulticastDelegate<void(
-			FWidgetPreviewToolkitStateBase* InOldState, FWidgetPreviewToolkitStateBase* InNewState)>;
-
-		virtual UWidgetPreview* GetPreview() const = 0;
-		virtual FWidgetPreviewToolkitStateBase* GetState() const = 0;
-		virtual FOnStateChanged& OnStateChanged() = 0;
-	};
-
 	class FWidgetPreviewToolkit
 		: public FBaseAssetToolkit
 		, public FGCObject
@@ -163,13 +144,20 @@ namespace UE::UMGWidgetPreview::Private
 		//~ End FGCObject
 
 		//~ Begin IWidgetPreviewToolkit
+		virtual TSharedPtr<FLayoutExtender> GetLayoutExtender() const override;
+		virtual FOnSelectedObjectsChanged& OnSelectedObjectsChanged() override;
+		virtual TConstArrayView<TWeakObjectPtr<UObject>> GetSelectedObjects() const override;
+		virtual void SetSelectedObjects(const TArray<TWeakObjectPtr<UObject>>& InObjects) override;
 		virtual UWidgetPreview* GetPreview() const override;
-		virtual FWidgetPreviewToolkitStateBase* GetState() const override;
-		virtual FOnStateChanged& OnStateChanged() override;
+		virtual UWorld* GetPreviewWorld() override;
 		//~ End IWidgetPreviewToolkit
 
+		FWidgetPreviewToolkitStateBase* GetState() const;
+
+		using FOnStateChanged = TMulticastDelegate<void(FWidgetPreviewToolkitStateBase* InOldState, FWidgetPreviewToolkitStateBase* InNewState)>;
+		FOnStateChanged& OnStateChanged();
+
 		TSharedPtr<FWidgetPreviewScene> GetPreviewScene();
-		UWorld* GetPreviewWorld();
 
 	public:
 		static const FLazyName PreviewSceneSettingsTabID;
@@ -178,6 +166,8 @@ namespace UE::UMGWidgetPreview::Private
 	protected:
 		bool ShouldUpdate() const;
 
+		void OnBlueprintPrecompile(UBlueprint* InBlueprint);
+
 		void OnWidgetChanged(const EWidgetPreviewWidgetChangeType InChangeType);
 
 		void OnFocusChanging(
@@ -191,6 +181,7 @@ namespace UE::UMGWidgetPreview::Private
 		/** Resolve and set the current state based on various conditions. */
 		void ResolveState();
 
+		/** Resets to the default state. */
 		void ResetPreview();
 
 		virtual TSharedRef<SDockTab> SpawnTab_Viewport(const FSpawnTabArgs& Args) override;
@@ -201,6 +192,9 @@ namespace UE::UMGWidgetPreview::Private
 	private:
 		TObjectPtr<UWidgetPreview> Preview;
 
+		FOnSelectedObjectsChanged SelectedObjectsChangedDelegate;
+		TArray<TWeakObjectPtr<UObject>> SelectedObjects;
+
 		TSharedPtr<FWidgetPreviewScene> PreviewScene;
 		FAdvancedPreviewSceneModule::FOnPreviewSceneChanged OnPreviewSceneChangedDelegate;
 		TSharedPtr<SWidget> PreviewSettingsWidget;
@@ -209,8 +203,8 @@ namespace UE::UMGWidgetPreview::Private
 		TSharedPtr<SWidget> MessageLogWidget;
 
 		bool bIsFocused = false;
-		FWorkflowAllowedTabSet TabFactories;
 
+		FDelegateHandle OnBlueprintPrecompileHandle;
 		FDelegateHandle OnWidgetChangedHandle;
 		FDelegateHandle OnFocusChangingHandle;
 		FOnStateChanged OnStateChangedDelegate;
diff --git a/Engine/Plugins/Editor/UMGWidgetPreview/Source/UMGWidgetPreview/Private/WidgetPreviewTypesPrivate.h b/Engine/Plugins/Editor/UMGWidgetPreview/Source/UMGWidgetPreview/Private/WidgetPreviewTypesPrivate.h
--- a/Engine/Plugins/Editor/UMGWidgetPreview/Source/UMGWidgetPreview/Private/WidgetPreviewTypesPrivate.h
+++ b/Engine/Plugins/Editor/UMGWidgetPreview/Source/UMGWidgetPreview/Private/WidgetPreviewTypesPrivate.h
@@ -10,7 +10,16 @@ namespace UE::UMGWidgetPreview::Private
 {
 	struct FWidgetTypeTuple
 	{
+		FWidgetTypeTuple() = default;
+
 		explicit FWidgetTypeTuple(const UUserWidget* InUserWidgetCDO)
+			: bIsValid(true)
+		{
+			Set(InUserWidgetCDO);
+		}
+
+		/** Attempt to resolve the tuple from the given UserWidget CDO. */
+		void Set(const UUserWidget* InUserWidgetCDO)
 		{
 			check(InUserWidgetCDO);
 
@@ -22,5 +31,15 @@ namespace UE::UMGWidgetPreview::Private
 		const UUserWidget* ClassDefaultObject = nullptr;
 		UWidgetBlueprint* Blueprint = nullptr;
 		UWidgetBlueprintGeneratedClass* BlueprintGeneratedClass = nullptr;
+
+		/** Returns true if at least one of the tuple values is valid. */
+		inline bool IsValid() const
+		{
+			return bIsValid && (ClassDefaultObject || Blueprint || BlueprintGeneratedClass);
+		}
+
+	private:
+		/** True only when value resolution has been attempted. */
+		bool bIsValid = false;
 	};
 }
diff --git a/Engine/Plugins/Editor/UMGWidgetPreview/Source/UMGWidgetPreview/Private/Widgets/SWidgetPreview.cpp b/Engine/Plugins/Editor/UMGWidgetPreview/Source/UMGWidgetPreview/Private/Widgets/SWidgetPreview.cpp
--- a/Engine/Plugins/Editor/UMGWidgetPreview/Source/UMGWidgetPreview/Private/Widgets/SWidgetPreview.cpp
+++ b/Engine/Plugins/Editor/UMGWidgetPreview/Source/UMGWidgetPreview/Private/Widgets/SWidgetPreview.cpp
@@ -26,7 +26,7 @@ namespace UE::UMGWidgetPreview::Private
 
 		ContainerWidget = SNew(SBorder)
 		[
-			CreatedSlateWidget.ToSharedRef()
+			GetCreatedSlateWidget()
 		];
 
 		OnWidgetChanged(EWidgetPreviewWidgetChangeType::Assignment);
@@ -48,6 +48,8 @@ namespace UE::UMGWidgetPreview::Private
 
 	SWidgetPreview::~SWidgetPreview()
 	{
+		ContainerWidget->ClearContent();
+
 		if (const TSharedPtr<FWidgetPreviewToolkit> Toolkit = WeakToolkit.Pin())
 		{
 			Toolkit->OnStateChanged().Remove(OnStateChangedHandle);
@@ -82,7 +84,7 @@ namespace UE::UMGWidgetPreview::Private
 			SWidgetPreview* MutableThis = const_cast<SWidgetPreview*>(this);
 
 			MutableThis->CreatedSlateWidget = SNullWidget::NullWidget;
-			ContainerWidget->SetContent(CreatedSlateWidget.ToSharedRef());
+			ContainerWidget->SetContent(GetCreatedSlateWidget());
 			MutableThis->bClearWidgetOnNextPaint = false;
 		}
 
@@ -111,22 +113,28 @@ namespace UE::UMGWidgetPreview::Private
 			return;
 		}
 
-		if (const TSharedPtr<FWidgetPreviewToolkit> Toolkit = WeakToolkit.Pin())
+		if (InChangeType != EWidgetPreviewWidgetChangeType::Destroyed)
 		{
-			if (UWidgetPreview* Preview = Toolkit->GetPreview())
+			if (const TSharedPtr<FWidgetPreviewToolkit> Toolkit = WeakToolkit.Pin())
 			{
-				UWorld* World = GetWorld();
-
-				if (UUserWidget* PreviewWidget = Preview->GetOrCreateWidgetInstance(World))
+				if (UWidgetPreview* Preview = Toolkit->GetPreview())
 				{
-					CreatedSlateWidget = PreviewWidget->TakeWidget();
+					UWorld* World = GetWorld();
+					if (const TSharedPtr<SWidget> PreviewSlateWidget = Preview->GetSlateWidgetInstance())
+					{
+						CreatedSlateWidget = PreviewSlateWidget;
+					}
+					else if (UUserWidget* PreviewWidget = Preview->GetOrCreateWidgetInstance(World))
+					{
+						CreatedSlateWidget = PreviewWidget->TakeWidget();
+					}
+					else
+					{
+						CreatedSlateWidget = SNullWidget::NullWidget;
+					}
+
+					ContainerWidget->SetContent(GetCreatedSlateWidget());
 				}
-				else
-				{
-					CreatedSlateWidget = SNullWidget::NullWidget;
-				}
-
-				ContainerWidget->SetContent(CreatedSlateWidget.ToSharedRef());
 			}
 		}
 	}
@@ -140,6 +148,16 @@ namespace UE::UMGWidgetPreview::Private
 
 		return nullptr;
 	}
+
+	TSharedRef<SWidget> SWidgetPreview::GetCreatedSlateWidget() const
+	{
+		if (TSharedPtr<SWidget> SlateWidget = CreatedSlateWidget.Pin())
+		{
+			return SlateWidget.ToSharedRef();
+		}
+
+		return SNullWidget::NullWidget;
+	}
 }
 
 #undef LOCTEXT_NAMESPACE
diff --git a/Engine/Plugins/Editor/UMGWidgetPreview/Source/UMGWidgetPreview/Private/Widgets/SWidgetPreview.h b/Engine/Plugins/Editor/UMGWidgetPreview/Source/UMGWidgetPreview/Private/Widgets/SWidgetPreview.h
--- a/Engine/Plugins/Editor/UMGWidgetPreview/Source/UMGWidgetPreview/Private/Widgets/SWidgetPreview.h
+++ b/Engine/Plugins/Editor/UMGWidgetPreview/Source/UMGWidgetPreview/Private/Widgets/SWidgetPreview.h
@@ -39,12 +39,14 @@ namespace UE::UMGWidgetPreview::Private
 		/** Convenience method to get world from the associated viewport. */
 		UWorld* GetWorld() const;
 
+		TSharedRef<SWidget> GetCreatedSlateWidget() const;
+
 	private:
 		TWeakPtr<FWidgetPreviewToolkit> WeakToolkit;
 
 		TSharedPtr<SRetainerWidget> RetainerWidget;
 		TSharedPtr<SBorder> ContainerWidget;
-		TSharedPtr<SWidget> CreatedSlateWidget;
+		TWeakPtr<SWidget> CreatedSlateWidget;
 
 		bool bClearWidgetOnNextPaint = false;
 		bool bIsRetainedRender = false;
diff --git a/Engine/Plugins/Editor/UMGWidgetPreview/Source/UMGWidgetPreview/Private/Widgets/SWidgetPreviewDetails.cpp b/Engine/Plugins/Editor/UMGWidgetPreview/Source/UMGWidgetPreview/Private/Widgets/SWidgetPreviewDetails.cpp
--- a/Engine/Plugins/Editor/UMGWidgetPreview/Source/UMGWidgetPreview/Private/Widgets/SWidgetPreviewDetails.cpp
+++ b/Engine/Plugins/Editor/UMGWidgetPreview/Source/UMGWidgetPreview/Private/Widgets/SWidgetPreviewDetails.cpp
@@ -2,16 +2,20 @@
 
 #include "SWidgetPreviewDetails.h"
 
+#include "Blueprint/UserWidget.h"
 #include "IDetailsView.h"
 #include "Modules/ModuleManager.h"
 #include "PropertyEditorModule.h"
 #include "WidgetPreview.h"
+#include "WidgetPreviewToolkit.h"
 
 namespace UE::UMGWidgetPreview::Private
 {
-	void SWidgetPreviewDetails::Construct(const FArguments& Args, UWidgetPreview* InPreview)
+	void SWidgetPreviewDetails::Construct(const FArguments& Args, const TSharedRef<FWidgetPreviewToolkit>& InToolkit)
 	{
-		WeakPreview = InPreview;
+		WeakToolkit = InToolkit;
+
+		OnSelectedObjectsChangedHandle = InToolkit->OnSelectedObjectsChanged().AddSP(this, &SWidgetPreviewDetails::OnSelectedObjectChanged);
 
 		FPropertyEditorModule& PropertyEditor = FModuleManager::LoadModuleChecked<FPropertyEditorModule>("PropertyEditor");
 
@@ -26,7 +30,7 @@ namespace UE::UMGWidgetPreview::Private
 
 		DetailsView = PropertyEditor.CreateDetailView(DetailsViewArgs);
 
-		HandleSelectedObjectChanged();
+		OnSelectedObjectChanged({});
 
 		ChildSlot
 		[
@@ -34,11 +38,71 @@ namespace UE::UMGWidgetPreview::Private
 		];
 	}
 
-	void SWidgetPreviewDetails::HandleSelectedObjectChanged()
+	SWidgetPreviewDetails::~SWidgetPreviewDetails()
+	{
+		if (const TSharedPtr<FWidgetPreviewToolkit> Toolkit = WeakToolkit.Pin())
+		{
+			Toolkit->OnSelectedObjectsChanged().Remove(OnSelectedObjectsChangedHandle);
+		}
+	}
+
+	void SWidgetPreviewDetails::OnSelectedObjectChanged(const TConstArrayView<TWeakObjectPtr<UObject>> InSelectedObjects) const
+	{
+		if (InSelectedObjects.IsEmpty())
+		{
+			if (const TSharedPtr<FWidgetPreviewToolkit> Toolkit = WeakToolkit.Pin())
+			{
+				if (UWidgetPreview* Preview = Toolkit->GetPreview())
+				{
+					DetailsView->SetObject(Preview);
+				}
+				else
+				{
+					DetailsView->SetObject(nullptr);
+				}
+			}
+		}
+		else
+		{
+			TArray<UObject*> SelectedObjects;
+			Algo::TransformIf(
+				InSelectedObjects,
+				SelectedObjects,
+				[](const TWeakObjectPtr<UObject>& InWeakObject)
+				{
+					return InWeakObject.IsValid();
+				},
+				[](const TWeakObjectPtr<UObject>& InWeakObject)
+				{
+					return InWeakObject.Get();
+				});
+
+			DetailsView->SetObjects(SelectedObjects);
+		}
+	}
+
+	void SWidgetPreviewDetails::NotifyPostChange(
+		const FPropertyChangedEvent& PropertyChangedEvent,
+		FEditPropertyChain* PropertyThatChanged)
 	{
-		if (UWidgetPreview* Preview = WeakPreview.Get())
+		for (int32 ObjectIdx = 0; ObjectIdx < PropertyChangedEvent.GetNumObjectsBeingEdited(); ++ObjectIdx)
 		{
-			DetailsView->SetObject(Preview);
+			if (const UObject* Object = PropertyChangedEvent.GetObjectBeingEdited(ObjectIdx))
+			{
+				if (PropertyThatChanged->GetActiveMemberNode())
+				{
+					const FProperty* Property = PropertyThatChanged->GetActiveMemberNode()->GetValue();
+					if (Property && Object->GetClass()->ImplementsInterface(UNotifyFieldValueChanged::StaticClass()))
+					{
+						TScriptInterface<INotifyFieldValueChanged> Interface = const_cast<UObject*>(Object);
+						UE::FieldNotification::FFieldId FieldId = Interface->GetFieldNotificationDescriptor().GetField(Object->GetClass(), Property->GetFName());
+						if (FieldId.IsValid())
+						{
+							Interface->BroadcastFieldValueChanged(FieldId);
+						}
+					}
+				}
+			}
 		}
 	}
 }
diff --git a/Engine/Plugins/Editor/UMGWidgetPreview/Source/UMGWidgetPreview/Private/Widgets/SWidgetPreviewDetails.h b/Engine/Plugins/Editor/UMGWidgetPreview/Source/UMGWidgetPreview/Private/Widgets/SWidgetPreviewDetails.h
--- a/Engine/Plugins/Editor/UMGWidgetPreview/Source/UMGWidgetPreview/Private/Widgets/SWidgetPreviewDetails.h
+++ b/Engine/Plugins/Editor/UMGWidgetPreview/Source/UMGWidgetPreview/Private/Widgets/SWidgetPreviewDetails.h
@@ -3,27 +3,37 @@
 #pragma once
 
 #include "Misc/NotifyHook.h"
+#include "WidgetPreview.h"
 #include "Widgets/SCompoundWidget.h"
 
 class FWidgetBlueprintEditor;
-class FWidgetPreviewToolkit;
 class IDetailsView;
 class UWidgetPreview;
 
 namespace UE::UMGWidgetPreview::Private
 {
-	class SWidgetPreviewDetails : public SCompoundWidget, public FNotifyHook
+	class FWidgetPreviewToolkit;
+
+	class SWidgetPreviewDetails
+		: public SCompoundWidget
+		, public FNotifyHook
 	{
 		SLATE_BEGIN_ARGS(SWidgetPreviewDetails) {}
 		SLATE_END_ARGS()
 
-		void Construct(const FArguments& Args, UWidgetPreview* InPreview);
+		void Construct(const FArguments& Args, const TSharedRef<FWidgetPreviewToolkit>& InToolkit);
+
+		virtual ~SWidgetPreviewDetails() override;
+
+		virtual void NotifyPostChange(const FPropertyChangedEvent& PropertyChangedEvent, class FEditPropertyChain* PropertyThatChanged) override;
 
 	private:
-		void HandleSelectedObjectChanged();
+		void OnSelectedObjectChanged(const TConstArrayView<TWeakObjectPtr<UObject>> InSelectedObjects) const;
 
 	private:
-		TWeakObjectPtr<UWidgetPreview> WeakPreview;
+		TWeakPtr<FWidgetPreviewToolkit> WeakToolkit;
 		TSharedPtr<IDetailsView> DetailsView;
+
+		FDelegateHandle OnSelectedObjectsChangedHandle;
 	};
 }
diff --git a/Engine/Plugins/Editor/UMGWidgetPreview/Source/UMGWidgetPreview/Public/IUMGWidgetPreviewModule.h b/Engine/Plugins/Editor/UMGWidgetPreview/Source/UMGWidgetPreview/Public/IUMGWidgetPreviewModule.h
new file mode 100644
--- /dev/null
+++ b/Engine/Plugins/Editor/UMGWidgetPreview/Source/UMGWidgetPreview/Public/IUMGWidgetPreviewModule.h
@@ -0,0 +1,25 @@
+// Copyright Epic Games, Inc. All Rights Reserved.
+
+#pragma once
+
+#include "Delegates/DelegateCombinations.h"
+#include "Framework/Docking/TabManager.h"
+#include "Modules/ModuleInterface.h"
+#include "Templates/SharedPointer.h"
+#include "UObject/NameTypes.h"
+
+extern const FLazyName UMGWidgetPreviewAppIdentifier;
+
+namespace UE::UMGWidgetPreview
+{
+	class IWidgetPreviewToolkit;
+}
+
+/** The public interface of the UMG widget preview (editor) module. */
+class IUMGWidgetPreviewModule
+	: public IModuleInterface
+{
+public:
+	DECLARE_EVENT_TwoParams(IUMGWidgetPreviewModule, FOnRegisterTabs, const TSharedPtr<UE::UMGWidgetPreview::IWidgetPreviewToolkit>&, const TSharedRef<FTabManager>&);
+	virtual FOnRegisterTabs& OnRegisterTabsForEditor() = 0;
+};
diff --git a/Engine/Plugins/Editor/UMGWidgetPreview/Source/UMGWidgetPreview/Public/IWidgetPreviewToolkit.h b/Engine/Plugins/Editor/UMGWidgetPreview/Source/UMGWidgetPreview/Public/IWidgetPreviewToolkit.h
new file mode 100644
--- /dev/null
+++ b/Engine/Plugins/Editor/UMGWidgetPreview/Source/UMGWidgetPreview/Public/IWidgetPreviewToolkit.h
@@ -0,0 +1,35 @@
+// Copyright Epic Games, Inc. All Rights Reserved.
+
+#pragma once
+
+#include "Framework/Docking/TabManager.h"
+#include "Logging/TokenizedMessage.h"
+#include "Templates/SharedPointer.h"
+#include "Toolkits/BaseToolkit.h"
+
+class UWidgetPreview;
+
+namespace UE::UMGWidgetPreview
+{
+	class IWidgetPreviewToolkit
+	{
+	public:
+		virtual ~IWidgetPreviewToolkit() = default;
+
+		virtual TSharedPtr<FLayoutExtender> GetLayoutExtender() const = 0;
+
+		using FOnSelectedObjectsChanged = TMulticastDelegate<void(const TConstArrayView<TWeakObjectPtr<UObject>>)>;
+		virtual FOnSelectedObjectsChanged& OnSelectedObjectsChanged() = 0;
+
+		/** Gets the "selected" objects, primarily used to display details for a given object/s other than the preview asset itself. */
+		virtual TConstArrayView<TWeakObjectPtr<UObject>> GetSelectedObjects() const = 0;
+
+		/** Sets the selected objects. When empty, this will reset to the Preview asset. */
+		virtual void SetSelectedObjects(const TArray<TWeakObjectPtr<UObject>>& InObjects) = 0;
+
+		virtual UWidgetPreview* GetPreview() const = 0;
+
+		/** Returns the Preview World. */
+		virtual UWorld* GetPreviewWorld() = 0;
+	};
+}
diff --git a/Engine/Plugins/Editor/UMGWidgetPreview/Source/UMGWidgetPreview/Public/WidgetPreview.h b/Engine/Plugins/Editor/UMGWidgetPreview/Source/UMGWidgetPreview/Public/WidgetPreview.h
new file mode 100644
--- /dev/null
+++ b/Engine/Plugins/Editor/UMGWidgetPreview/Source/UMGWidgetPreview/Public/WidgetPreview.h
@@ -0,0 +1,153 @@
+// Copyright Epic Games, Inc. All Rights Reserved.
+
+#pragma once
+
+#include "UObject/Object.h"
+
+#include "WidgetPreview.generated.h"
+
+class SWidget;
+class UPanelWidget;
+class UUserWidget;
+class UWidget;
+class UWidgetBlueprint;
+class UWidgetPreview;
+struct FImage;
+
+enum class EWidgetPreviewWidgetChangeType : uint8
+{
+	Assignment = 0,
+	Reinstanced = 1,
+	Structure = 2,
+	ChildReference = 3,
+	Destroyed = 4				// Just before the Slate widget is destroyed, etc.
+};
+
+USTRUCT(BlueprintType)
+struct UMGWIDGETPREVIEW_API FPreviewableWidgetVariant
+{
+	GENERATED_BODY();
+
+	UPROPERTY(BlueprintReadOnly, EditAnywhere, Category = "Widget", DisplayName = "Widget Type", meta = (AllowedClasses = "/Script/UMGEditor.WidgetBlueprint, /Script/UMGWidgetPreview.WidgetPreview"))
+	FSoftObjectPath ObjectPath;
+
+	FPreviewableWidgetVariant() = default;
+	explicit FPreviewableWidgetVariant(const TSubclassOf<UUserWidget>& InWidgetType);
+	explicit FPreviewableWidgetVariant(const UWidgetPreview* InWidgetPreview);
+
+public:
+	/** Flushes cached widgets and re-resolves from the ObjectPath. */
+	void UpdateCachedWidget();
+
+	/** Returns the referenced Object as a UUserWidget (CDO). Returns nullptr if not found, or we couldn't find a nested UUserWidget (ie. inside a UWidgetPreview). */
+	const UUserWidget* AsUserWidgetCDO() const;
+
+	/** Returns the referenced Object as a UWidgetPreview. Returns nullptr if not found, or not a UWidgetPreview. */
+	const UWidgetPreview* AsWidgetPreview() const;
+
+	friend bool operator==(const FPreviewableWidgetVariant& Left, const FPreviewableWidgetVariant& Right)
+	{
+		return Left.ObjectPath == Right.ObjectPath;
+	}
+
+	friend bool operator!=(const FPreviewableWidgetVariant& Left, const FPreviewableWidgetVariant& Right)
+	{
+		return !(Left == Right);
+	}
+
+private:
+	UPROPERTY(Transient)
+	TObjectPtr<const UUserWidget> CachedWidgetCDO;
+
+	UPROPERTY(Transient)
+	TWeakObjectPtr<UWidgetPreview> CachedWidgetPreview;
+};
+
+UCLASS(BlueprintType, NotBlueprintable, AutoExpandCategories = "Widgets")
+class UMGWIDGETPREVIEW_API UWidgetPreview
+	: public UObject
+{
+	GENERATED_BODY()
+
+public:
+	virtual void BeginDestroy() override;
+
+	/** Convenience function to check that all utilized widgets have bCanCallInitializedWithoutPlayerContext set to true, and reports any that don't. */
+	bool CanCallInitializedWithoutPlayerContext(const bool bInRecursive, TArray<const UUserWidget*>& OutFailedWidgets);
+
+	// @todo: move to utility func somewhere else?
+	/** Convenience function to check that the provided widget (and it's children) has bCanCallInitializedWithoutPlayerContext set to true, and reports any that don't. */
+	static bool CanCallInitializedWithoutPlayerContextOnWidget(const UUserWidget* InUserWidget, const bool bInRecursive, TArray<const UUserWidget*>& OutFailedWidgets);
+
+public:
+	using FOnWidgetChanged = TMulticastDelegate<void(const EWidgetPreviewWidgetChangeType)>;
+
+	FOnWidgetChanged& OnWidgetChanged() { return OnWidgetChangedDelegate; }
+
+	UFUNCTION(BlueprintCallable, Category = "Layout")
+	const TArray<FName>& GetWidgetSlotNames() const;
+
+	/** Returns or builds and returns an instance of the root widget for previewing. Can be used to trigger a rebuild. */
+	[[maybe_unused]] UUserWidget* GetOrCreateWidgetInstance(UWorld* InWorld, const bool bInForceRecreate = false);
+
+	/** Returns the current widget instance, if any. */
+	UUserWidget* GetWidgetInstance() const;
+
+	/** Returns the current underlying slate widget instance, if any. */
+	TSharedPtr<SWidget> GetSlateWidgetInstance() const;
+
+	/** Stores the current instance in PreviousWidgetInstance, and clears WidgetInstance. */
+	void ClearWidgetInstance();
+
+	const UUserWidget* GetWidgetCDO() const;
+	const UUserWidget* GetWidgetCDOForSlot(const FName InSlotName) const;
+
+	const FPreviewableWidgetVariant& GetWidgetType() const;
+	void SetWidgetType(const FPreviewableWidgetVariant& InWidget);
+
+	const TMap<FName, FPreviewableWidgetVariant>& GetSlotWidgetTypes() const;
+	void SetSlotWidgetTypes(const TMap<FName, FPreviewableWidgetVariant>& InWidgets);
+
+protected:
+	virtual void PostLoad() override;
+
+	virtual void PostEditChangeProperty(FPropertyChangedEvent& PropertyChangedEvent) override;
+
+	void OnWidgetBlueprintChanged(UBlueprint* InBlueprint);
+
+	/** Misc. functionality to perform after a widget assignment is changed. */
+	void UpdateWidgets();
+
+	/** Creates a new WidgetInstance, replacing the current one if it exists. */
+	UUserWidget* CreateWidgetInstance(UWorld* InWorld);
+
+	void CleanupReferences();
+
+	/** Returns slot names not already occupied in SlotWidgets. */
+	UFUNCTION(BlueprintCallable, Category = "Widget")
+	TArray<FName> GetAvailableWidgetSlotNames();
+
+private:
+	/** Widget to preview. */
+	UPROPERTY(EditAnywhere, BlueprintReadWrite, Getter, Setter, Category = "Widget", DisplayName = "Widget", meta = (AllowPrivateAccess = "true", ShowOnlyInnerProperties))
+	FPreviewableWidgetVariant WidgetType;
+
+	/** Widget per-slot, if WidgetType has any. */
+	UPROPERTY(EditAnywhere, BlueprintReadWrite, Getter, Setter, Category = "Widget", DisplayName = "Slot Widgets", meta = (AllowPrivateAccess = "true", GetKeyOptions = "GetAvailableWidgetSlotNames", ShowOnlyInnerProperties))
+	TMap<FName, FPreviewableWidgetVariant> SlotWidgetTypes;
+
+	UPROPERTY(DuplicateTransient)
+	TObjectPtr<UUserWidget> WidgetInstance;
+
+	TSharedPtr<SWidget> SlateWidgetInstance;
+
+	/** Slot names available in WidgetType (if any). */
+	UPROPERTY(Transient)
+	TArray<FName> SlotNameCache;
+
+	/** Widgets here should be checked for validity when a new one is assigned, to allow tear-down functionality. */
+	UPROPERTY(Transient)
+	TArray<TWeakObjectPtr<const UUserWidget>> WidgetReferenceCache;
+
+	FOnWidgetChanged OnWidgetChangedDelegate;
+};
