[UMG Widget Preview] Module destruction fix

diff --git a/Engine/Plugins/Editor/UMGWidgetPreview/Source/UMGWidgetPreview/Private/UMGWidgetPreviewModule.cpp b/Engine/Plugins/Editor/UMGWidgetPreview/Source/UMGWidgetPreview/Private/UMGWidgetPreviewModule.cpp
--- a/Engine/Plugins/Editor/UMGWidgetPreview/Source/UMGWidgetPreview/Private/UMGWidgetPreviewModule.cpp
+++ b/Engine/Plugins/Editor/UMGWidgetPreview/Source/UMGWidgetPreview/Private/UMGWidgetPreviewModule.cpp
@@ -16,6 +16,7 @@
 #include "ToolMenuMisc.h"
 #include "ToolMenus.h"
 #include "ToolMenuSection.h"
+#include "UObject/NameTypes.h"
 #include "UObject/UObjectHash.h"
 #include "WidgetBlueprint.h"
 #include "WidgetPreview.h"
@@ -30,6 +31,9 @@ DEFINE_LOG_CATEGORY(LogWidgetPreview);
 
 namespace UE::UMGWidgetPreview::Private
 {
+	static const FLazyName PreviewWidgetVariantName("PreviewableWidgetVariant");
+	static const FLazyName WidgetPreviewName("WidgetPreview");
+
 	void FUMGWidgetPreviewModule::StartupModule()
 	{
 		FWidgetPreviewCommands::Register();
@@ -49,11 +53,11 @@ namespace UE::UMGWidgetPreview::Private
 		FPropertyEditorModule& PropertyModule = FModuleManager::GetModuleChecked<FPropertyEditorModule>("PropertyEditor");
 		{
 			PropertyModule.RegisterCustomPropertyTypeLayout(
-				FPreviewableWidgetVariant::StaticStruct()->GetFName(),
+				PreviewWidgetVariantName,
 				FOnGetPropertyTypeCustomizationInstance::CreateStatic(&FPreviewableWidgetCustomization::MakeInstance));
 
 			PropertyModule.RegisterCustomClassLayout(
-				UWidgetPreview::StaticClass()->GetFName(),
+				WidgetPreviewName,
 				FOnGetDetailCustomizationInstance::CreateStatic(&FWidgetPreviewCustomization::MakeInstance));
 		}
 	}
@@ -64,8 +68,8 @@ namespace UE::UMGWidgetPreview::Private
 
 		if (FPropertyEditorModule* PropertyModule = FModuleManager::GetModulePtr<FPropertyEditorModule>("PropertyEditor"))
 		{
-			PropertyModule->UnregisterCustomClassLayout(UWidgetPreview::StaticClass()->GetFName());
-			PropertyModule->UnregisterCustomPropertyTypeLayout(FPreviewableWidgetVariant::StaticStruct()->GetFName());
+			PropertyModule->UnregisterCustomClassLayout(WidgetPreviewName);
+			PropertyModule->UnregisterCustomPropertyTypeLayout(PreviewWidgetVariantName);
 		}
 
 		if (FMessageLogModule* MessageLogModule = FModuleManager::GetModulePtr<FMessageLogModule>("MessageLog"))
