Updated Widget Previews to respect the preview size specified in the UMG designer (Desired, Fill, etc.) and display the current preview size in the corner
Added a Widget Size override property to explicitly set the Preview Widget's size

diff --git a/Engine/Plugins/Editor/UMGWidgetPreview/Source/UMGWidgetPreview/Private/Customizations/PreviewableWidgetCustomization.cpp b/Engine/Plugins/Editor/UMGWidgetPreview/Source/UMGWidgetPreview/Private/Customizations/PreviewableWidgetCustomization.cpp
--- a/Engine/Plugins/Editor/UMGWidgetPreview/Source/UMGWidgetPreview/Private/Customizations/PreviewableWidgetCustomization.cpp
+++ b/Engine/Plugins/Editor/UMGWidgetPreview/Source/UMGWidgetPreview/Private/Customizations/PreviewableWidgetCustomization.cpp
@@ -358,7 +358,10 @@ namespace UE::UMGWidgetPreview::Private
 
 	void FPreviewableWidgetCustomization::OnWidgetChanged(const EWidgetPreviewWidgetChangeType InChangeType)
 	{
-		PropertyUtilities->RequestForceRefresh();
+		if (InChangeType != EWidgetPreviewWidgetChangeType::Resized)
+		{
+			PropertyUtilities->RequestForceRefresh();
+		}
 	}
 }
 
diff --git a/Engine/Plugins/Editor/UMGWidgetPreview/Source/UMGWidgetPreview/Private/Customizations/WidgetPreviewCustomization.cpp b/Engine/Plugins/Editor/UMGWidgetPreview/Source/UMGWidgetPreview/Private/Customizations/WidgetPreviewCustomization.cpp
--- a/Engine/Plugins/Editor/UMGWidgetPreview/Source/UMGWidgetPreview/Private/Customizations/WidgetPreviewCustomization.cpp
+++ b/Engine/Plugins/Editor/UMGWidgetPreview/Source/UMGWidgetPreview/Private/Customizations/WidgetPreviewCustomization.cpp
@@ -13,6 +13,7 @@ namespace UE::UMGWidgetPreview::Private
 	const FName FWidgetPreviewCustomization::WidgetTypePropertyName = "WidgetType";
 	const FName FWidgetPreviewCustomization::SlotWidgetTypesPropertyName = "SlotWidgetTypes";
 	const FName FWidgetPreviewCustomization::WidgetInstancePropertyName = "WidgetInstance";
+	const FName FWidgetPreviewCustomization::OverriddenSizePropertyName = "OverriddenWidgetSize";
 
 	TSharedRef<IDetailCustomization> FWidgetPreviewCustomization::MakeInstance()
 	{
@@ -27,6 +28,7 @@ namespace UE::UMGWidgetPreview::Private
 		TSharedRef<IPropertyHandle> WidgetTypePropertyHandle = DetailBuilder.GetProperty(WidgetTypePropertyName);
 		TSharedRef<IPropertyHandle> SlotWidgetTypesPropertyHandle = DetailBuilder.GetProperty(SlotWidgetTypesPropertyName);
 		TSharedRef<IPropertyHandle> WidgetInstancePropertyHandle = DetailBuilder.GetProperty(WidgetInstancePropertyName);
+		TSharedRef<IPropertyHandle> OverriddenSizePropertyHandle = DetailBuilder.GetProperty(OverriddenSizePropertyName);
 
 		static const FName WidgetCategoryName = "Widget";
 
@@ -36,6 +38,7 @@ namespace UE::UMGWidgetPreview::Private
 			{
 				IDetailCategoryBuilder& WidgetCategory = DetailBuilder.EditCategory(WidgetCategoryName);
 				WidgetCategory.AddProperty(WidgetTypePropertyHandle);
+				WidgetCategory.AddProperty(OverriddenSizePropertyHandle);
 
 				// Widget Instance
 				{
diff --git a/Engine/Plugins/Editor/UMGWidgetPreview/Source/UMGWidgetPreview/Private/Customizations/WidgetPreviewCustomization.h b/Engine/Plugins/Editor/UMGWidgetPreview/Source/UMGWidgetPreview/Private/Customizations/WidgetPreviewCustomization.h
--- a/Engine/Plugins/Editor/UMGWidgetPreview/Source/UMGWidgetPreview/Private/Customizations/WidgetPreviewCustomization.h
+++ b/Engine/Plugins/Editor/UMGWidgetPreview/Source/UMGWidgetPreview/Private/Customizations/WidgetPreviewCustomization.h
@@ -36,5 +36,6 @@ namespace UE::UMGWidgetPreview::Private
 		static const FName WidgetTypePropertyName;
         static const FName SlotWidgetTypesPropertyName;
         static const FName WidgetInstancePropertyName;
+        static const FName OverriddenSizePropertyName;
 	};
 }
diff --git a/Engine/Plugins/Editor/UMGWidgetPreview/Source/UMGWidgetPreview/Private/WidgetPreview.cpp b/Engine/Plugins/Editor/UMGWidgetPreview/Source/UMGWidgetPreview/Private/WidgetPreview.cpp
--- a/Engine/Plugins/Editor/UMGWidgetPreview/Source/UMGWidgetPreview/Private/WidgetPreview.cpp
+++ b/Engine/Plugins/Editor/UMGWidgetPreview/Source/UMGWidgetPreview/Private/WidgetPreview.cpp
@@ -105,6 +105,13 @@ const UWidgetPreview* FPreviewableWidgetVariant::AsWidgetPreview() const
 	return nullptr;
 }
 
+UWidgetPreview::UWidgetPreview(const FObjectInitializer& ObjectInitializer)
+	: Super(ObjectInitializer)
+	, bShouldOverrideWidgetSize(false)
+{
+	OverriddenWidgetSize = FVector2D(100,100);
+}
+
 const TArray<FName>& UWidgetPreview::GetWidgetSlotNames() const
 {
 	return SlotNameCache;
@@ -354,6 +361,35 @@ void UWidgetPreview::SetSlotWidgetTypes(const TMap<FName, FPreviewableWidgetVari
 	}
 }
 
+const bool UWidgetPreview::GetbShouldOverrideWidgetSize() const
+{
+	return bShouldOverrideWidgetSize;
+}
+
+void UWidgetPreview::SetbShouldOverrideWidgetSize(bool InOverride)
+{
+	if (bShouldOverrideWidgetSize != InOverride)
+	{
+		bShouldOverrideWidgetSize = InOverride;
+		OnWidgetChanged().Broadcast(EWidgetPreviewWidgetChangeType::Resized);
+	}
+}
+
+const FVector2D UWidgetPreview::GetOverriddenWidgetSize() const
+{
+	return OverriddenWidgetSize;
+}
+
+void UWidgetPreview::SetOverriddenWidgetSize(FVector2D InWidgetSize)
+{
+	if (OverriddenWidgetSize != InWidgetSize)
+	{
+		OverriddenWidgetSize = InWidgetSize;
+		OnWidgetChanged().Broadcast(EWidgetPreviewWidgetChangeType::Resized);
+	}
+	
+}
+
 void UWidgetPreview::PostLoad()
 {
 	UObject::PostLoad();
@@ -374,6 +410,11 @@ void UWidgetPreview::PostEditChangeProperty(FPropertyChangedEvent& PropertyChang
 		UpdateWidgets();
 		OnWidgetChanged().Broadcast(EWidgetPreviewWidgetChangeType::Assignment);
 	}
+	else if (PropertyName == GET_MEMBER_NAME_CHECKED(UWidgetPreview, bShouldOverrideWidgetSize)
+			|| PropertyName == GET_MEMBER_NAME_CHECKED(UWidgetPreview, OverriddenWidgetSize))
+	{
+		OnWidgetChanged().Broadcast(EWidgetPreviewWidgetChangeType::Resized);
+	}
 }
 
 // Due to compilation, etc.
diff --git a/Engine/Plugins/Editor/UMGWidgetPreview/Source/UMGWidgetPreview/Private/Widgets/SWidgetPreview.cpp b/Engine/Plugins/Editor/UMGWidgetPreview/Source/UMGWidgetPreview/Private/Widgets/SWidgetPreview.cpp
--- a/Engine/Plugins/Editor/UMGWidgetPreview/Source/UMGWidgetPreview/Private/Widgets/SWidgetPreview.cpp
+++ b/Engine/Plugins/Editor/UMGWidgetPreview/Source/UMGWidgetPreview/Private/Widgets/SWidgetPreview.cpp
@@ -2,7 +2,6 @@
 
 #include "SWidgetPreview.h"
 
-#include "Blueprint/UserWidget.h"
 #include "Editor.h"
 #include "Engine/Engine.h"
 #include "Engine/World.h"
@@ -29,21 +28,36 @@ namespace UE::UMGWidgetPreview::Private
 			GetCreatedSlateWidget()
 		];
 
-		OnWidgetChanged(EWidgetPreviewWidgetChangeType::Assignment);
-
 		ChildSlot
 		[
 			SNew(SOverlay)
 			+ SOverlay::Slot()
+			.Expose(OverlaySlot)
 			[
-				SAssignNew(RetainerWidget, SRetainerWidget)
-				.RenderOnPhase(false)
-				.RenderOnInvalidation(false)
+				SAssignNew(SizeBoxWidget, SBox)
 				[
-					ContainerWidget.ToSharedRef()
+					SAssignNew(RetainerWidget, SRetainerWidget)
+					.RenderOnPhase(false)
+					.RenderOnInvalidation(false)
+					[
+						ContainerWidget.ToSharedRef()
+					]
 				]
 			]
+			+ SOverlay::Slot()
+			.HAlign(HAlign_Right)
+			.VAlign(VAlign_Bottom)
+			.Padding(0.0f, 0.0f, 6.0f, 2.0f)
+			[
+				SNew(STextBlock)
+				.TextStyle(FAppStyle::Get(), "Graph.ZoomText")
+				.Text(this, &SWidgetPreview::GetPreviewSizeText)
+				.ColorAndOpacity(FSlateColor(FLinearColor(1, 1, 1, 0.25f)))
+			]
 		];
+
+		OnWidgetChanged(EWidgetPreviewWidgetChangeType::Assignment);
+
 	}
 
 	SWidgetPreview::~SWidgetPreview()
@@ -113,7 +127,30 @@ namespace UE::UMGWidgetPreview::Private
 			return;
 		}
 
-		if (InChangeType != EWidgetPreviewWidgetChangeType::Destroyed)
+		if (InChangeType == EWidgetPreviewWidgetChangeType::Resized)
+		{
+			if (const TSharedPtr<FWidgetPreviewToolkit> Toolkit = WeakToolkit.Pin())
+			{
+				if (UWidgetPreview* Preview = Toolkit->GetPreview())
+				{
+					// Override with custom size if set
+					bSizeOverridden = Preview->GetbShouldOverrideWidgetSize();
+					if (bSizeOverridden)
+					{
+						CreatedWidgetPreviewSizeMode = EDesignPreviewSizeMode::Custom;
+						CreatedWidgetPreviewCustomSize = Preview->GetOverriddenWidgetSize();
+						
+					}
+					else
+					{
+						CreatedWidgetPreviewSizeMode = Preview->GetWidgetCDO()->DesignSizeMode;
+						CreatedWidgetPreviewCustomSize = Preview->GetWidgetCDO()->DesignTimeSize;
+					}
+					RebuildContainer();
+				}
+			}
+		}
+		else if (InChangeType != EWidgetPreviewWidgetChangeType::Destroyed)
 		{
 			if (const TSharedPtr<FWidgetPreviewToolkit> Toolkit = WeakToolkit.Pin())
 			{
@@ -133,6 +170,15 @@ namespace UE::UMGWidgetPreview::Private
 						CreatedSlateWidget = SNullWidget::NullWidget;
 					}
 
+					// Override with custom size if set
+					bSizeOverridden = Preview->GetbShouldOverrideWidgetSize();
+					if (!bSizeOverridden)
+					{
+						CreatedWidgetPreviewSizeMode = Preview->GetWidgetCDO()->DesignSizeMode;
+						CreatedWidgetPreviewCustomSize = Preview->GetWidgetCDO()->DesignTimeSize;
+					}
+
+					RebuildContainer();
 					ContainerWidget->SetContent(GetCreatedSlateWidget());
 				}
 			}
@@ -158,6 +204,101 @@ namespace UE::UMGWidgetPreview::Private
 
 		return SNullWidget::NullWidget;
 	}
+
+	void SWidgetPreview::RebuildContainer()
+	{
+		switch (CreatedWidgetPreviewSizeMode)
+		{
+			case EDesignPreviewSizeMode::Custom:
+			case EDesignPreviewSizeMode::CustomOnScreen:
+			{
+				if (SizeBoxWidget.IsValid())
+				{
+					SizeBoxWidget->SetWidthOverride(CreatedWidgetPreviewCustomSize.X);
+					SizeBoxWidget->SetHeightOverride(CreatedWidgetPreviewCustomSize.Y);
+				}
+				OverlaySlot->SetHorizontalAlignment(HAlign_Center);
+				OverlaySlot->SetVerticalAlignment(VAlign_Center);
+				break;
+			}
+			case EDesignPreviewSizeMode::Desired:
+			case EDesignPreviewSizeMode::DesiredOnScreen:
+			{
+				if (SizeBoxWidget.IsValid())
+				{
+					SizeBoxWidget->SetWidthOverride(FOptionalSize());
+					SizeBoxWidget->SetHeightOverride(FOptionalSize());
+				}
+				OverlaySlot->SetHorizontalAlignment(HAlign_Center);
+				OverlaySlot->SetVerticalAlignment(VAlign_Center);
+				break;
+			}
+			case EDesignPreviewSizeMode::FillScreen:
+			{
+				if (SizeBoxWidget.IsValid())
+				{
+					SizeBoxWidget->SetWidthOverride(FOptionalSize());
+					SizeBoxWidget->SetHeightOverride(FOptionalSize());
+				}
+				OverlaySlot->SetHorizontalAlignment(HAlign_Fill);
+				OverlaySlot->SetVerticalAlignment(VAlign_Fill);
+				break;
+			}
+			default:
+			{
+				checkf(false, TEXT("Invalid Widget Preview Size Mode"));
+				break;
+			}
+		}
+	}
+
+	FText SWidgetPreview::GetPreviewSizeText() const
+	{
+		
+
+		switch (CreatedWidgetPreviewSizeMode)
+		{
+			case EDesignPreviewSizeMode::Custom:
+			case EDesignPreviewSizeMode::CustomOnScreen:
+			{
+				FNumberFormattingOptions Options = FNumberFormattingOptions::DefaultNoGrouping();
+
+				if (bSizeOverridden)
+				{
+					return FText::Format(LOCTEXT("PreviewSizeTextCustom", "Preview Size: Overriden ({0}, {1})"),
+						FText::AsNumber(CreatedWidgetPreviewCustomSize.X, &FNumberFormattingOptions::DefaultNoGrouping()),
+						FText::AsNumber(CreatedWidgetPreviewCustomSize.Y, &FNumberFormattingOptions::DefaultNoGrouping()));
+				}
+				else
+				{
+					return FText::Format(LOCTEXT("PreviewSizeTextCustom", "Preview Size: Custom ({0}, {1})"),
+						FText::AsNumber(CreatedWidgetPreviewCustomSize.X, &FNumberFormattingOptions::DefaultNoGrouping()),
+						FText::AsNumber(CreatedWidgetPreviewCustomSize.Y, &FNumberFormattingOptions::DefaultNoGrouping()));
+				}
+				break;
+			}
+			case EDesignPreviewSizeMode::Desired:
+			case EDesignPreviewSizeMode::DesiredOnScreen:
+			{
+				return FText(LOCTEXT("PreviewSizeTextDesired", "Preview Size: Desired"));
+				break;
+			}
+			case EDesignPreviewSizeMode::FillScreen:
+			{
+				return FText(LOCTEXT("PreviewSizeTextFill", "Preview Size: Fill Screen"));
+				break;
+			}
+			default:
+			{
+				checkf(false, TEXT("Invalid Widget Preview Size Mode"));
+				break;
+			}
+		}
+
+
+		return FText();
+	}
+
 }
 
 #undef LOCTEXT_NAMESPACE
diff --git a/Engine/Plugins/Editor/UMGWidgetPreview/Source/UMGWidgetPreview/Private/Widgets/SWidgetPreview.h b/Engine/Plugins/Editor/UMGWidgetPreview/Source/UMGWidgetPreview/Private/Widgets/SWidgetPreview.h
--- a/Engine/Plugins/Editor/UMGWidgetPreview/Source/UMGWidgetPreview/Private/Widgets/SWidgetPreview.h
+++ b/Engine/Plugins/Editor/UMGWidgetPreview/Source/UMGWidgetPreview/Private/Widgets/SWidgetPreview.h
@@ -2,8 +2,11 @@
 
 #pragma once
 
+#include "Blueprint/UserWidget.h"
 #include "UObject/StrongObjectPtr.h"
 #include "UObject/ObjectPtr.h"
+#include "Widgets/Layout/SBox.h"
+#include "Widgets/SOverlay.h"
 #include "WidgetPreview.h"
 #include "Widgets/SCompoundWidget.h"
 
@@ -41,17 +44,26 @@ namespace UE::UMGWidgetPreview::Private
 
 		TSharedRef<SWidget> GetCreatedSlateWidget() const;
 
+		void RebuildContainer();
+		FText GetPreviewSizeText() const;
+
 	private:
 		TWeakPtr<FWidgetPreviewToolkit> WeakToolkit;
 
 		TSharedPtr<SRetainerWidget> RetainerWidget;
 		TSharedPtr<SBorder> ContainerWidget;
+		TSharedPtr<SBox> SizeBoxWidget;
 		TWeakPtr<SWidget> CreatedSlateWidget;
+		SOverlay::FOverlaySlot* OverlaySlot;
 
 		bool bClearWidgetOnNextPaint = false;
 		bool bIsRetainedRender = false;
+		bool bSizeOverridden = false;
 
 		FDelegateHandle OnStateChangedHandle;
 		FDelegateHandle OnWidgetChangedHandle;
+
+		EDesignPreviewSizeMode CreatedWidgetPreviewSizeMode;
+		FVector2D CreatedWidgetPreviewCustomSize;
 	};
 }
diff --git a/Engine/Plugins/Editor/UMGWidgetPreview/Source/UMGWidgetPreview/Public/WidgetPreview.h b/Engine/Plugins/Editor/UMGWidgetPreview/Source/UMGWidgetPreview/Public/WidgetPreview.h
--- a/Engine/Plugins/Editor/UMGWidgetPreview/Source/UMGWidgetPreview/Public/WidgetPreview.h
+++ b/Engine/Plugins/Editor/UMGWidgetPreview/Source/UMGWidgetPreview/Public/WidgetPreview.h
@@ -20,7 +20,8 @@ enum class EWidgetPreviewWidgetChangeType : uint8
 	Reinstanced = 1,
 	Structure = 2,
 	ChildReference = 3,
-	Destroyed = 4				// Just before the Slate widget is destroyed, etc.
+	Destroyed = 4,				// Just before the Slate widget is destroyed, etc.
+	Resized = 5
 };
 
 USTRUCT(BlueprintType)
@@ -70,6 +71,9 @@ class UMGWIDGETPREVIEW_API UWidgetPreview
 	GENERATED_BODY()
 
 public:
+
+	UWidgetPreview(const FObjectInitializer& ObjectInitializer);
+
 	virtual void BeginDestroy() override;
 
 	/** Convenience function to check that all utilized widgets have bCanCallInitializedWithoutPlayerContext set to true, and reports any that don't. */
@@ -108,6 +112,12 @@ public:
 	const TMap<FName, FPreviewableWidgetVariant>& GetSlotWidgetTypes() const;
 	void SetSlotWidgetTypes(const TMap<FName, FPreviewableWidgetVariant>& InWidgets);
 
+	const bool GetbShouldOverrideWidgetSize() const;
+	void SetbShouldOverrideWidgetSize(bool InOverride);
+
+	const FVector2D GetOverriddenWidgetSize() const;
+	void SetOverriddenWidgetSize(FVector2D InWidgetSize);
+
 protected:
 	virtual void PostLoad() override;
 
@@ -136,6 +146,13 @@ private:
 	UPROPERTY(EditAnywhere, BlueprintReadWrite, Getter, Setter, Category = "Widget", DisplayName = "Slot Widgets", meta = (AllowPrivateAccess = "true", GetKeyOptions = "GetAvailableWidgetSlotNames", ShowOnlyInnerProperties))
 	TMap<FName, FPreviewableWidgetVariant> SlotWidgetTypes;
 
+	/** Widget Custom Size Override */
+	UPROPERTY(EditAnywhere, BlueprintReadWrite, Getter, Setter, Category = "Widget", DisplayName = "Override Widget Size", meta = (AllowPrivateAccess = "true", InlineEditConditionToggle))
+	bool bShouldOverrideWidgetSize;
+
+	UPROPERTY(EditAnywhere, BlueprintReadWrite, Getter, Setter, Category = "Widget", DisplayName = "Widget Size", meta = (AllowPrivateAccess = "true", EditCondition="bShouldOverrideWidgetSize"))
+	FVector2D OverriddenWidgetSize;
+
 	UPROPERTY(DuplicateTransient)
 	TObjectPtr<UUserWidget> WidgetInstance;
 
