Widget Preview CDO fix

diff --git a/Engine/Plugins/Editor/UMGWidgetPreview/Source/UMGWidgetPreview/Private/WidgetPreview.cpp b/Engine/Plugins/Editor/UMGWidgetPreview/Source/UMGWidgetPreview/Private/WidgetPreview.cpp
--- a/Engine/Plugins/Editor/UMGWidgetPreview/Source/UMGWidgetPreview/Private/WidgetPreview.cpp
+++ b/Engine/Plugins/Editor/UMGWidgetPreview/Source/UMGWidgetPreview/Private/WidgetPreview.cpp
@@ -35,8 +35,13 @@ void FPreviewableWidgetVariant::UpdateCachedWidget()
     	if (UWidgetPreview* WidgetPreview = Cast<UWidgetPreview>(ResolvedObject))
     	{
     		CachedWidgetPreview = WidgetPreview;
-    		PreviewUserWidgetCDO = WidgetPreview->GetWidgetCDO();
-    		PreviewUserWidgetTuple.Set(WidgetPreview->GetWidgetCDO());
+
+    		// If we have a widget preview, it may not necessarily have a valid referenced widget type to preview
+    		if (const UUserWidget* WidgetPreviewWidgetCDO = WidgetPreview->GetWidgetCDO())
+    		{
+    			PreviewUserWidgetCDO = WidgetPreviewWidgetCDO;
+    			PreviewUserWidgetTuple.Set(WidgetPreviewWidgetCDO);
+    		}
     	}
 		else if (const UWidgetBlueprint* AsBlueprint = Cast<UWidgetBlueprint>(ResolvedObject))
 		{
@@ -338,6 +343,14 @@ void UWidgetPreview::SetWidgetType(const FPreviewableWidgetVariant& InWidget)
 	{
 		WidgetType = InWidget;
 		ClearWidgetInstance();
+
+		// Check for self-reference
+		if (WidgetType.ObjectPath == this)
+		{
+			UE_LOG(LogWidgetPreview, Warning, TEXT("Tried to reference self as the widget type. This is not allowed."));
+			WidgetType.ObjectPath.Reset();
+		}
+
 		UpdateWidgets();
 
 		OnWidgetChanged().Broadcast(EWidgetPreviewWidgetChangeType::Assignment);
@@ -407,6 +420,14 @@ void UWidgetPreview::PostEditChangeProperty(FPropertyChangedEvent& PropertyChang
 		|| PropertyName.IsNone()) // None can be an Undo operation
 	{
 		ClearWidgetInstance();
+
+		// Check for self-reference and disallow
+		if (WidgetType.ObjectPath == this)
+		{
+			UE_LOG(LogWidgetPreview, Warning, TEXT("Tried to reference self as the widget type. This is not allowed."));
+			WidgetType.ObjectPath.Reset();
+		}
+
 		UpdateWidgets();
 		OnWidgetChanged().Broadcast(EWidgetPreviewWidgetChangeType::Assignment);
 	}
diff --git a/Engine/Plugins/Editor/UMGWidgetPreview/Source/UMGWidgetPreview/Private/Widgets/SWidgetPreview.cpp b/Engine/Plugins/Editor/UMGWidgetPreview/Source/UMGWidgetPreview/Private/Widgets/SWidgetPreview.cpp
--- a/Engine/Plugins/Editor/UMGWidgetPreview/Source/UMGWidgetPreview/Private/Widgets/SWidgetPreview.cpp
+++ b/Engine/Plugins/Editor/UMGWidgetPreview/Source/UMGWidgetPreview/Private/Widgets/SWidgetPreview.cpp
@@ -139,12 +139,14 @@ namespace UE::UMGWidgetPreview::Private
 					{
 						CreatedWidgetPreviewSizeMode = EDesignPreviewSizeMode::Custom;
 						CreatedWidgetPreviewCustomSize = Preview->GetOverriddenWidgetSize();
-						
 					}
-					else
+					else if (!Preview->GetWidgetType().ObjectPath.IsNull())
 					{
-						CreatedWidgetPreviewSizeMode = Preview->GetWidgetCDO()->DesignSizeMode;
-						CreatedWidgetPreviewCustomSize = Preview->GetWidgetCDO()->DesignTimeSize;
+						if (const UUserWidget* WidgetCDO = Preview->GetWidgetCDO())
+						{
+							CreatedWidgetPreviewSizeMode = WidgetCDO->DesignSizeMode;
+							CreatedWidgetPreviewCustomSize = WidgetCDO->DesignTimeSize;
+						}
 					}
 					RebuildContainer();
 				}
@@ -157,6 +159,7 @@ namespace UE::UMGWidgetPreview::Private
 				if (UWidgetPreview* Preview = Toolkit->GetPreview())
 				{
 					UWorld* World = GetWorld();
+					bool bHasValidWidget = true;
 					if (const TSharedPtr<SWidget> PreviewSlateWidget = Preview->GetSlateWidgetInstance())
 					{
 						CreatedSlateWidget = PreviewSlateWidget;
@@ -168,14 +171,18 @@ namespace UE::UMGWidgetPreview::Private
 					else
 					{
 						CreatedSlateWidget = SNullWidget::NullWidget;
+						bHasValidWidget = false;
 					}
 
 					// Override with custom size if set
 					bSizeOverridden = Preview->GetbShouldOverrideWidgetSize();
-					if (!bSizeOverridden)
+					if (bHasValidWidget && !bSizeOverridden)
 					{
-						CreatedWidgetPreviewSizeMode = Preview->GetWidgetCDO()->DesignSizeMode;
-						CreatedWidgetPreviewCustomSize = Preview->GetWidgetCDO()->DesignTimeSize;
+						if (const UUserWidget* WidgetCDO = Preview->GetWidgetCDO())
+						{
+							CreatedWidgetPreviewSizeMode = WidgetCDO->DesignSizeMode;
+							CreatedWidgetPreviewCustomSize = WidgetCDO->DesignTimeSize;
+						}
 					}
 
 					RebuildContainer();
