Minor fix to a potential nullptr access in UMG Widget Preview Editor Plugin

diff --git a/Engine/Plugins/Editor/UMGWidgetPreview/Source/UMGWidgetPreview/Private/Customizations/PreviewableWidgetCustomization.cpp b/Engine/Plugins/Editor/UMGWidgetPreview/Source/UMGWidgetPreview/Private/Customizations/PreviewableWidgetCustomization.cpp
--- a/Engine/Plugins/Editor/UMGWidgetPreview/Source/UMGWidgetPreview/Private/Customizations/PreviewableWidgetCustomization.cpp
+++ b/Engine/Plugins/Editor/UMGWidgetPreview/Source/UMGWidgetPreview/Private/Customizations/PreviewableWidgetCustomization.cpp
@@ -97,34 +97,37 @@ namespace UE::UMGWidgetPreview::Private
 
 			return bCanCall;
 		};
-
-		UClass* WidgetClass = PreviewableWidgetVariant->AsUserWidgetCDO()->GetClass();
-
-		// Get all of the functions we need to display (done ahead of time so we can sort them)
-		for (TFieldIterator<UFunction> FunctionIter(WidgetClass, EFieldIteratorFlags::IncludeSuper); FunctionIter; ++FunctionIter)
+		
+		if (const UUserWidget* UserWidgetCDO = PreviewableWidgetVariant->AsUserWidgetCDO())
 		{
-			UFunction* TestFunction = *FunctionIter;
-			if (CanDisplayAndCallFunction(TestFunction))
+			UClass* WidgetClass = UserWidgetCDO->GetClass();
+
+			// Get all of the functions we need to display (done ahead of time so we can sort them)
+			for (TFieldIterator<UFunction> FunctionIter(WidgetClass, EFieldIteratorFlags::IncludeSuper); FunctionIter; ++FunctionIter)
 			{
-				bool bAllowFunction = true;
-				if (UClass* TestFunctionOwnerClass = TestFunction->GetOwnerClass())
+				UFunction* TestFunction = *FunctionIter;
+				if (CanDisplayAndCallFunction(TestFunction))
 				{
-					if (UBlueprint* Blueprint = Cast<UBlueprint>(TestFunctionOwnerClass->ClassGeneratedBy))
+					bool bAllowFunction = true;
+					if (UClass* TestFunctionOwnerClass = TestFunction->GetOwnerClass())
 					{
-						if (FBlueprintEditorUtils::IsEditorUtilityBlueprint(Blueprint))
+						if (UBlueprint* Blueprint = Cast<UBlueprint>(TestFunctionOwnerClass->ClassGeneratedBy))
 						{
-							// Skip Blutilities if disabled via project settings
-							bAllowFunction = !bDisallowEditorUtilityBlueprintFunctions;
+							if (FBlueprintEditorUtils::IsEditorUtilityBlueprint(Blueprint))
+							{
+								// Skip Blutilities if disabled via project settings
+								bAllowFunction = !bDisallowEditorUtilityBlueprintFunctions;
+							}
 						}
 					}
-				}
 
-				if (bAllowFunction)
-				{
-					const FName FunctionName = TestFunction->GetFName();
-					if (!OutCallInEditorFunctions.FindByPredicate([&FunctionName](const UFunction* Func) { return Func->GetFName() == FunctionName; }))
+					if (bAllowFunction)
 					{
-						OutCallInEditorFunctions.Add(*FunctionIter);
+						const FName FunctionName = TestFunction->GetFName();
+						if (!OutCallInEditorFunctions.FindByPredicate([&FunctionName](const UFunction* Func) { return Func->GetFName() == FunctionName; }))
+						{
+							OutCallInEditorFunctions.Add(*FunctionIter);
+						}
 					}
 				}
 			}
