MVVM Widget Preview: Fixed status message lagging behind state when changing source widgets.

#Synced-CL 51601552

diff --git a/Engine/Plugins/Editor/UMGWidgetPreview/Source/UMGWidgetPreview/Private/WidgetPreviewToolkit.cpp b/Engine/Plugins/Editor/UMGWidgetPreview/Source/UMGWidgetPreview/Private/WidgetPreviewToolkit.cpp
--- a/Engine/Plugins/Editor/UMGWidgetPreview/Source/UMGWidgetPreview/Private/WidgetPreviewToolkit.cpp
+++ b/Engine/Plugins/Editor/UMGWidgetPreview/Source/UMGWidgetPreview/Private/WidgetPreviewToolkit.cpp
@@ -659,6 +659,11 @@ namespace UE::UMGWidgetPreview::Private
 		return OnStateChangedDelegate;
 	}
 
+	FSimpleMulticastDelegate& FWidgetPreviewToolkit::OnMessageChanged()
+	{
+		return OnMessageChangedDelegate;
+	}
+
 	UWorld* FWidgetPreviewToolkit::GetPreviewWorld()
 	{
 		if (TSharedPtr<FWidgetPreviewScene> PreviewScenePtr = GetPreviewScene())
@@ -734,31 +739,36 @@ namespace UE::UMGWidgetPreview::Private
 		}
 	}
 
-	void FWidgetPreviewToolkit::SetState(FWidgetPreviewToolkitStateBase* InNewState)
+	bool FWidgetPreviewToolkit::SetState(FWidgetPreviewToolkitStateBase* InNewState)
 	{
 		FWidgetPreviewToolkitStateBase* OldState = CurrentState;
 		FWidgetPreviewToolkitStateBase* NewState = InNewState;
 
-		if (OldState != NewState)
+		if (OldState == NewState)
 		{
-			if (OldState)
-			{
-				OldState->OnExit(NewState);
-			}
+			return false;
+		}
 
-			if (NewState)
-			{
-				NewState->OnEnter(OldState);
-			}
+		if (OldState)
+		{
+			OldState->OnExit(NewState);
+		}
 
-			CurrentState = NewState;
-			OnStateChanged().Broadcast(OldState, NewState);
+		if (NewState)
+		{
+			NewState->OnEnter(OldState);
 		}
+
+		CurrentState = NewState;
+		OnStateChanged().Broadcast(OldState, NewState);
+
+		return true;
 	}
 
 	void FWidgetPreviewToolkit::ResolveState()
 	{
 		FWidgetPreviewToolkitStateBase* NewState = nullptr;
+		bool bForceMessageUpdate = false;
 
 		if (!bIsFocused)
 		{
@@ -771,6 +781,7 @@ namespace UE::UMGWidgetPreview::Private
 			{
 				UnsupportedWidgetState.SetUnsupportedWidgets(FailedWidgets);
 				NewState = &UnsupportedWidgetState;
+				bForceMessageUpdate = true;
 			}
 
 			// If we're here, the current state should be valid/running
@@ -780,7 +791,12 @@ namespace UE::UMGWidgetPreview::Private
 			}
 		}
 
-		SetState(NewState);
+		const bool bStateChanged = SetState(NewState);
+
+		if (!bStateChanged && bForceMessageUpdate)
+		{
+			OnMessageChanged().Broadcast();
+		}
 	}
 
 	void FWidgetPreviewToolkit::ResetPreview()
diff --git a/Engine/Plugins/Editor/UMGWidgetPreview/Source/UMGWidgetPreview/Private/WidgetPreviewToolkit.h b/Engine/Plugins/Editor/UMGWidgetPreview/Source/UMGWidgetPreview/Private/WidgetPreviewToolkit.h
--- a/Engine/Plugins/Editor/UMGWidgetPreview/Source/UMGWidgetPreview/Private/WidgetPreviewToolkit.h
+++ b/Engine/Plugins/Editor/UMGWidgetPreview/Source/UMGWidgetPreview/Private/WidgetPreviewToolkit.h
@@ -156,6 +156,7 @@ namespace UE::UMGWidgetPreview::Private
 
 		using FOnStateChanged = TMulticastDelegate<void(FWidgetPreviewToolkitStateBase* InOldState, FWidgetPreviewToolkitStateBase* InNewState)>;
 		FOnStateChanged& OnStateChanged();
+		FSimpleMulticastDelegate& OnMessageChanged();
 
 		TSharedPtr<FWidgetPreviewScene> GetPreviewScene();
 
@@ -175,8 +176,8 @@ namespace UE::UMGWidgetPreview::Private
 			const FWeakWidgetPath& InOldWidgetPath, const TSharedPtr<SWidget>& InOldWidget,
 			const FWidgetPath& InNewWidgetPath, const TSharedPtr<SWidget>& InNewWidget);
 
-		/** If the given state is different to the current state, this will handle transitions and events. */
-		void SetState(FWidgetPreviewToolkitStateBase* InNewState);
+		/** If the given state is different to the current state, this will handle transitions and events. Returns true if the state changed. */
+		bool SetState(FWidgetPreviewToolkitStateBase* InNewState);
 
 		/** Resolve and set the current state based on various conditions. */
 		void ResolveState();
@@ -208,6 +209,7 @@ namespace UE::UMGWidgetPreview::Private
 		FDelegateHandle OnWidgetChangedHandle;
 		FDelegateHandle OnFocusChangingHandle;
 		FOnStateChanged OnStateChangedDelegate;
+		FSimpleMulticastDelegate OnMessageChangedDelegate; // Called when state doesn't change, but the message still needs to update.
 
 		FWidgetPreviewToolkitStateBase* CurrentState = nullptr;
 		FWidgetPreviewToolkitPausedState PausedState;
diff --git a/Engine/Plugins/Editor/UMGWidgetPreview/Source/UMGWidgetPreview/Private/Widgets/SWidgetPreviewStatus.cpp b/Engine/Plugins/Editor/UMGWidgetPreview/Source/UMGWidgetPreview/Private/Widgets/SWidgetPreviewStatus.cpp
--- a/Engine/Plugins/Editor/UMGWidgetPreview/Source/UMGWidgetPreview/Private/Widgets/SWidgetPreviewStatus.cpp
+++ b/Engine/Plugins/Editor/UMGWidgetPreview/Source/UMGWidgetPreview/Private/Widgets/SWidgetPreviewStatus.cpp
@@ -16,6 +16,7 @@ namespace UE::UMGWidgetPreview::Private
 		WeakToolkit = InToolkit;
 
 		OnStateChangedHandle = InToolkit->OnStateChanged().AddSP(this, &SWidgetPreviewStatus::OnStateChanged);
+		OnMessageChangedHandle = InToolkit->OnMessageChanged().AddSP(this, &SWidgetPreviewStatus::OnMessageChanged);
 
 		MessageContainerWidget = SNew(SBox);
 		MessageContainerWidget->SetContent(MakeMessageWidget());
@@ -54,6 +55,7 @@ namespace UE::UMGWidgetPreview::Private
 		if (const TSharedPtr<FWidgetPreviewToolkit> Toolkit = WeakToolkit.Pin())
 		{
 			Toolkit->OnStateChanged().Remove(OnStateChangedHandle);
+			Toolkit->OnMessageChanged().Remove(OnMessageChangedHandle);
 		}
 	}
 
@@ -62,6 +64,11 @@ namespace UE::UMGWidgetPreview::Private
 		MessageContainerWidget->SetContent(MakeMessageWidget());
 	}
 
+	void SWidgetPreviewStatus::OnMessageChanged()
+	{
+		MessageContainerWidget->SetContent(MakeMessageWidget());
+	}
+
 	TSharedRef<SWidget> SWidgetPreviewStatus::MakeMessageWidget()
 	{
 		if (const TSharedPtr<FTokenizedMessage>& StatusMessage = GetStatusMessage())
diff --git a/Engine/Plugins/Editor/UMGWidgetPreview/Source/UMGWidgetPreview/Private/Widgets/SWidgetPreviewStatus.h b/Engine/Plugins/Editor/UMGWidgetPreview/Source/UMGWidgetPreview/Private/Widgets/SWidgetPreviewStatus.h
--- a/Engine/Plugins/Editor/UMGWidgetPreview/Source/UMGWidgetPreview/Private/Widgets/SWidgetPreviewStatus.h
+++ b/Engine/Plugins/Editor/UMGWidgetPreview/Source/UMGWidgetPreview/Private/Widgets/SWidgetPreviewStatus.h
@@ -26,6 +26,7 @@ namespace UE::UMGWidgetPreview::Private
 
 	private:
 		void OnStateChanged(FWidgetPreviewToolkitStateBase* InOldState, FWidgetPreviewToolkitStateBase* InNewState);
+		void OnMessageChanged();
 
 		TSharedRef<SWidget> MakeMessageWidget();
 
@@ -41,5 +42,6 @@ namespace UE::UMGWidgetPreview::Private
 		TSharedPtr<SBox> MessageContainerWidget;
 
 		FDelegateHandle OnStateChangedHandle;
+		FDelegateHandle OnMessageChangedHandle;
 	};
 }
